Updating template defs and manifests to resolve duplicate class type references listed when creating a new project

diff --git a/Templates/TP_FirstPersonBP/Config/TemplateDefs.ini b/Templates/TP_FirstPersonBP/Config/TemplateDefs.ini
--- a/Templates/TP_FirstPersonBP/Config/TemplateDefs.ini
+++ b/Templates/TP_FirstPersonBP/Config/TemplateDefs.ini
@@ -5,7 +5,7 @@ LocalizedDescriptions=(Language="en",Text="The first person template features a
 LocalizedDisplayNames=(Language="ko",Text="일인칭")
 LocalizedDescriptions=(Language="ko",Text="일인칭 템플릿은 일인칭 시점에서 보이는 (양팔로 표현되는) 플레이어 캐릭터를 특징으로 합니다. 캐릭터 이동은 키보드, 콘트롤러, 터치 디바이스의 가상 조이스틱으로 가능합니다. 추가적으로 플레이어는 마우스, 콘트롤러, 터치 디바이스의 가상 조이스틱을 사용하여 둘러보기도 가능합니다. 캐릭터에는 총이 장착되어 있어, 마우스, 콘트롤러, 터치 디바이스의 가상 조이스틱을 사용하면 간단한 구체 프로젝타일이 발사되어, 아레나 벽에 튕기면서 레벨에 있는 몇몇 물리 오브젝트에 영향을 끼칩니다.")
 
-ClassTypes=Character, GameMode, HUD, Actor, Character, GameMode, HUD, Actor
+ClassTypes=Character, GameMode, HUD, Actor
 AssetTypes=Animation Sequence, Animation Montage, Sound Wave, Skeletal Mesh, Animation Blueprint, Skeleton, Static Mesh, Texture.
 
 FoldersToIgnore=Binaries
diff --git a/Templates/TP_VirtualRealityBP/Config/TemplateDefs.ini b/Templates/TP_VirtualRealityBP/Config/TemplateDefs.ini
--- a/Templates/TP_VirtualRealityBP/Config/TemplateDefs.ini
+++ b/Templates/TP_VirtualRealityBP/Config/TemplateDefs.ini
@@ -5,7 +5,7 @@ LocalizedDescriptions=(Language="en",Text="Blueprint Virtual Reality Template fo
 LocalizedDisplayNames=(Language="ko",Text="가상 현실")
 LocalizedDescriptions=(Language="ko",Text="데스크탑 및 Playstation VR 용 블루프린트 가상 현실 템플릿입니다. 각 플랫폼 별 전용 폰과 맵은 물론 두 가지 텔레포트 이동 방식이 포함되어 있습니다. 플레이어는 MotionControllerMap 에서 모션 컨트롤러를 사용하여 텔레포트할 수 있으며, HMDLocomotionMap 에서 게임패드 아니면 HeadMountedDisplay 만으로 이동할 곳을 지정하는 방식으로 텔레포트할 수 있습니다.")
 
-ClassTypes=Character, GameMode, HUD, Actor, Character, GameMode, HUD, Actor
+ClassTypes=Character, GameMode, HUD, Actor
 AssetTypes=Animation Sequence, Animation Montage, Sound Wave, Skeletal Mesh, Animation Blueprint, Skeleton, Static Mesh, Texture.
 
 FoldersToIgnore=Binaries
diff --git a/Templates/TP_VirtualRealityBP/manifest.json b/Templates/TP_VirtualRealityBP/manifest.json
--- a/Templates/TP_VirtualRealityBP/manifest.json
+++ b/Templates/TP_VirtualRealityBP/manifest.json
@@ -24,7 +24,7 @@
       "Text": "VR,Virtual,Reality,Virtual Reality"
     }
   ],
-  "ClassTypes": "Character, GameMode, HUD, Actor, Character, GameMode, HUD, Actor",
+  "ClassTypes": "Character, GameMode, HUD, Actor",
   "Category": "BlueprintFeature",
   "FocusAsset": "/Game/VirtualRealityBP/VirtualReality.VirtualReality",
   "Thumbnail": "TP_VirtualRealityBP.png",
