Map right controller in addition to left for Magic Leap in First Person templates

diff --git a/Templates/TP_FirstPerson/Config/DefaultInput.ini b/Templates/TP_FirstPerson/Config/DefaultInput.ini
--- a/Templates/TP_FirstPerson/Config/DefaultInput.ini
+++ b/Templates/TP_FirstPerson/Config/DefaultInput.ini
@@ -149,6 +149,9 @@ DoubleClickTime=0.200000
 +ActionMappings=(ActionName="ResetVR",bShift=False,bCtrl=False,bAlt=False,bCmd=False,Key=OculusTouch_Left_Thumbstick_Click)
 +ActionMappings=(ActionName="ResetVR",bShift=False,bCtrl=False,bAlt=False,bCmd=False,Key=ValveIndex_Left_Thumbstick_Click)
 +ActionMappings=(ActionName="ResetVR",bShift=False,bCtrl=False,bAlt=False,bCmd=False,Key=MagicLeap_Left_Bumper)
++ActionMappings=(ActionName="Jump",bShift=False,bCtrl=False,bAlt=False,bCmd=False,Key=MagicLeap_Right_Trigger)
++ActionMappings=(ActionName="Fire",bShift=False,bCtrl=False,bAlt=False,bCmd=False,Key=MagicLeap_Right_Trackpad_Touch)
++ActionMappings=(ActionName="ResetVR",bShift=False,bCtrl=False,bAlt=False,bCmd=False,Key=MagicLeap_Right_Bumper)
 +AxisMappings=(AxisName="MoveForward",Scale=1.000000,Key=W)
 +AxisMappings=(AxisName="MoveForward",Scale=-1.000000,Key=S)
 +AxisMappings=(AxisName="MoveForward",Scale=1.000000,Key=Up)
@@ -177,6 +180,8 @@ DoubleClickTime=0.200000
 +AxisMappings=(AxisName="MoveRight",Scale=1.000000,Key=ValveIndex_Left_Thumbstick_Y)
 +AxisMappings=(AxisName="MoveRight",Scale=1.000000,Key=MagicLeap_Left_Trackpad_X)
 +AxisMappings=(AxisName="MoveForward",Scale=1.000000,Key=MagicLeap_Left_Trackpad_Y)
++AxisMappings=(AxisName="MoveForward",Scale=1.000000,Key=MagicLeap_Right_Trackpad_Y)
++AxisMappings=(AxisName="MoveRight",Scale=1.000000,Key=MagicLeap_Right_Trackpad_X)
 DefaultTouchInterface=/Engine/MobileResources/HUD/DefaultVirtualJoysticks.DefaultVirtualJoysticks
 +ConsoleKeys=Tilde
 
diff --git a/Templates/TP_FirstPersonBP/Config/DefaultInput.ini b/Templates/TP_FirstPersonBP/Config/DefaultInput.ini
--- a/Templates/TP_FirstPersonBP/Config/DefaultInput.ini
+++ b/Templates/TP_FirstPersonBP/Config/DefaultInput.ini
@@ -149,6 +149,9 @@ DoubleClickTime=0.200000
 +ActionMappings=(ActionName="ResetVR",bShift=False,bCtrl=False,bAlt=False,bCmd=False,Key=OculusTouch_Left_Thumbstick_Click)
 +ActionMappings=(ActionName="ResetVR",bShift=False,bCtrl=False,bAlt=False,bCmd=False,Key=ValveIndex_Left_Thumbstick_Click)
 +ActionMappings=(ActionName="ResetVR",bShift=False,bCtrl=False,bAlt=False,bCmd=False,Key=MagicLeap_Left_Bumper)
++ActionMappings=(ActionName="Jump",bShift=False,bCtrl=False,bAlt=False,bCmd=False,Key=MagicLeap_Right_Trigger)
++ActionMappings=(ActionName="Fire",bShift=False,bCtrl=False,bAlt=False,bCmd=False,Key=MagicLeap_Right_Trackpad_Touch)
++ActionMappings=(ActionName="ResetVR",bShift=False,bCtrl=False,bAlt=False,bCmd=False,Key=MagicLeap_Right_Bumper)
 +AxisMappings=(AxisName="MoveForward",Scale=1.000000,Key=W)
 +AxisMappings=(AxisName="MoveForward",Scale=-1.000000,Key=S)
 +AxisMappings=(AxisName="MoveForward",Scale=1.000000,Key=Up)
@@ -177,6 +180,8 @@ DoubleClickTime=0.200000
 +AxisMappings=(AxisName="MoveRight",Scale=1.000000,Key=ValveIndex_Left_Thumbstick_Y)
 +AxisMappings=(AxisName="MoveRight",Scale=1.000000,Key=MagicLeap_Left_Trackpad_X)
 +AxisMappings=(AxisName="MoveForward",Scale=1.000000,Key=MagicLeap_Left_Trackpad_Y)
++AxisMappings=(AxisName="MoveForward",Scale=1.000000,Key=MagicLeap_Right_Trackpad_Y)
++AxisMappings=(AxisName="MoveRight",Scale=1.000000,Key=MagicLeap_Right_Trackpad_X)
 DefaultTouchInterface=/Engine/MobileResources/HUD/DefaultVirtualJoysticks.DefaultVirtualJoysticks
 +ConsoleKeys=Tilde
 
