Updating TemplateResources with FPWeapon changes
Updated VRTemplate to use FPWeapon
Updated FirstPerson Templates to use FPWeapon

diff --git a/Templates/TP_FirstPerson/Config/TemplateDefs.ini b/Templates/TP_FirstPerson/Config/TemplateDefs.ini
--- a/Templates/TP_FirstPerson/Config/TemplateDefs.ini
+++ b/Templates/TP_FirstPerson/Config/TemplateDefs.ini
@@ -48,3 +48,4 @@ ReplacementsInFiles=(Extensions=("cpp","h","ini","cs"),From="%TEMPLATENAME_LOWER
 ReplacementsInFiles=(Extensions=("cpp","h","ini","cs"),From="%TEMPLATENAME%",To="%PROJECTNAME%",bCaseSensitive=false)
 SharedContentPacks=(MountName="FirstPerson",DetailLevels=("Standard"))
 SharedContentPacks=(MountName="Geometry",DetailLevels=("Standard","High"))
+SharedContentPacks=(MountName="FPWeapon",DetailLevels=("Standard"))
diff --git a/Templates/TP_FirstPersonBP/Config/TemplateDefs.ini b/Templates/TP_FirstPersonBP/Config/TemplateDefs.ini
--- a/Templates/TP_FirstPersonBP/Config/TemplateDefs.ini
+++ b/Templates/TP_FirstPersonBP/Config/TemplateDefs.ini
@@ -40,4 +40,5 @@ ReplacementsInFiles=(Extensions=("cpp","h","ini","cs"),From="%TEMPLATENAME_LOWER
 ReplacementsInFiles=(Extensions=("cpp","h","ini","cs"),From="%TEMPLATENAME%",To="%PROJECTNAME%",bCaseSensitive=false)
 SharedContentPacks=(MountName="FirstPerson",DetailLevels=("Standard"))
 SharedContentPacks=(MountName="Geometry",DetailLevels=("Standard","High"))
+SharedContentPacks=(MountName="FPWeapon",DetailLevels=("Standard"))
 EditDetailLevelPreference="High"
\ No newline at end of file
