Removed VR functionality from TP_FirstPerson and TP_FirstPersonBP. Updated config to import shared template resources (FPWeapon) when template is added as a feature pack. Cleaned up comments and blueprints.

diff --git a/Templates/TP_FirstPerson/Config/DefaultInput.ini b/Templates/TP_FirstPerson/Config/DefaultInput.ini
--- a/Templates/TP_FirstPerson/Config/DefaultInput.ini
+++ b/Templates/TP_FirstPerson/Config/DefaultInput.ini
@@ -1,103 +1,30 @@
-
 [/Script/Engine.InputSettings]
+-AxisConfig=(AxisKeyName="Gamepad_LeftX",AxisProperties=(DeadZone=0.25,Exponent=1.f,Sensitivity=1.f))
+-AxisConfig=(AxisKeyName="Gamepad_LeftY",AxisProperties=(DeadZone=0.25,Exponent=1.f,Sensitivity=1.f))
+-AxisConfig=(AxisKeyName="Gamepad_RightX",AxisProperties=(DeadZone=0.25,Exponent=1.f,Sensitivity=1.f))
+-AxisConfig=(AxisKeyName="Gamepad_RightY",AxisProperties=(DeadZone=0.25,Exponent=1.f,Sensitivity=1.f))
+-AxisConfig=(AxisKeyName="MouseX",AxisProperties=(DeadZone=0.f,Exponent=1.f,Sensitivity=0.07f))
+-AxisConfig=(AxisKeyName="MouseY",AxisProperties=(DeadZone=0.f,Exponent=1.f,Sensitivity=0.07f))
+-AxisConfig=(AxisKeyName="Mouse2D",AxisProperties=(DeadZone=0.f,Exponent=1.f,Sensitivity=0.07f))
++AxisConfig=(AxisKeyName="Mouse2D",AxisProperties=(DeadZone=0.000000,Sensitivity=0.070000,Exponent=1.000000,bInvert=False))
 +AxisConfig=(AxisKeyName="Gamepad_LeftX",AxisProperties=(DeadZone=0.250000,Sensitivity=1.000000,Exponent=1.000000,bInvert=False))
 +AxisConfig=(AxisKeyName="Gamepad_LeftY",AxisProperties=(DeadZone=0.250000,Sensitivity=1.000000,Exponent=1.000000,bInvert=False))
 +AxisConfig=(AxisKeyName="Gamepad_RightX",AxisProperties=(DeadZone=0.250000,Sensitivity=1.000000,Exponent=1.000000,bInvert=False))
 +AxisConfig=(AxisKeyName="Gamepad_RightY",AxisProperties=(DeadZone=0.250000,Sensitivity=1.000000,Exponent=1.000000,bInvert=False))
 +AxisConfig=(AxisKeyName="MouseX",AxisProperties=(DeadZone=0.000000,Sensitivity=0.070000,Exponent=1.000000,bInvert=False))
 +AxisConfig=(AxisKeyName="MouseY",AxisProperties=(DeadZone=0.000000,Sensitivity=0.070000,Exponent=1.000000,bInvert=False))
-+AxisConfig=(AxisKeyName="MotionController_Left_Thumbstick_Z",AxisProperties=(DeadZone=0.000000,Sensitivity=1.000000,Exponent=1.000000,bInvert=False))
-+AxisConfig=(AxisKeyName="MotionController_Right_Thumbstick_Z",AxisProperties=(DeadZone=0.000000,Sensitivity=1.000000,Exponent=1.000000,bInvert=False))
 +AxisConfig=(AxisKeyName="MouseWheelAxis",AxisProperties=(DeadZone=0.000000,Sensitivity=1.000000,Exponent=1.000000,bInvert=False))
 +AxisConfig=(AxisKeyName="Gamepad_LeftTriggerAxis",AxisProperties=(DeadZone=0.000000,Sensitivity=1.000000,Exponent=1.000000,bInvert=False))
 +AxisConfig=(AxisKeyName="Gamepad_RightTriggerAxis",AxisProperties=(DeadZone=0.000000,Sensitivity=1.000000,Exponent=1.000000,bInvert=False))
-+AxisConfig=(AxisKeyName="MotionController_Left_Thumbstick_X",AxisProperties=(DeadZone=0.000000,Sensitivity=1.000000,Exponent=1.000000,bInvert=False))
-+AxisConfig=(AxisKeyName="MotionController_Left_Thumbstick_Y",AxisProperties=(DeadZone=0.000000,Sensitivity=1.000000,Exponent=1.000000,bInvert=False))
-+AxisConfig=(AxisKeyName="MotionController_Left_TriggerAxis",AxisProperties=(DeadZone=0.000000,Sensitivity=1.000000,Exponent=1.000000,bInvert=False))
-+AxisConfig=(AxisKeyName="MotionController_Left_Grip1Axis",AxisProperties=(DeadZone=0.000000,Sensitivity=1.000000,Exponent=1.000000,bInvert=False))
-+AxisConfig=(AxisKeyName="MotionController_Left_Grip2Axis",AxisProperties=(DeadZone=0.000000,Sensitivity=1.000000,Exponent=1.000000,bInvert=False))
-+AxisConfig=(AxisKeyName="MotionController_Right_Thumbstick_X",AxisProperties=(DeadZone=0.000000,Sensitivity=1.000000,Exponent=1.000000,bInvert=False))
-+AxisConfig=(AxisKeyName="MotionController_Right_Thumbstick_Y",AxisProperties=(DeadZone=0.000000,Sensitivity=1.000000,Exponent=1.000000,bInvert=False))
-+AxisConfig=(AxisKeyName="MotionController_Right_TriggerAxis",AxisProperties=(DeadZone=0.000000,Sensitivity=1.000000,Exponent=1.000000,bInvert=False))
-+AxisConfig=(AxisKeyName="MotionController_Right_Grip1Axis",AxisProperties=(DeadZone=0.000000,Sensitivity=1.000000,Exponent=1.000000,bInvert=False))
-+AxisConfig=(AxisKeyName="MotionController_Right_Grip2Axis",AxisProperties=(DeadZone=0.000000,Sensitivity=1.000000,Exponent=1.000000,bInvert=False))
 +AxisConfig=(AxisKeyName="Gamepad_Special_Left_X",AxisProperties=(DeadZone=0.000000,Sensitivity=1.000000,Exponent=1.000000,bInvert=False))
 +AxisConfig=(AxisKeyName="Gamepad_Special_Left_Y",AxisProperties=(DeadZone=0.000000,Sensitivity=1.000000,Exponent=1.000000,bInvert=False))
-+AxisConfig=(AxisKeyName="OculusTouch_Left_Thumbstick",AxisProperties=(DeadZone=0.000000,Sensitivity=1.000000,Exponent=1.000000,bInvert=False))
-+AxisConfig=(AxisKeyName="OculusTouch_Left_FaceButton1",AxisProperties=(DeadZone=0.000000,Sensitivity=1.000000,Exponent=1.000000,bInvert=False))
-+AxisConfig=(AxisKeyName="OculusTouch_Left_Trigger",AxisProperties=(DeadZone=0.000000,Sensitivity=1.000000,Exponent=1.000000,bInvert=False))
-+AxisConfig=(AxisKeyName="OculusTouch_Left_FaceButton2",AxisProperties=(DeadZone=0.000000,Sensitivity=1.000000,Exponent=1.000000,bInvert=False))
-+AxisConfig=(AxisKeyName="OculusTouch_Left_IndexPointing",AxisProperties=(DeadZone=0.000000,Sensitivity=1.000000,Exponent=1.000000,bInvert=False))
-+AxisConfig=(AxisKeyName="OculusTouch_Left_ThumbUp",AxisProperties=(DeadZone=0.000000,Sensitivity=1.000000,Exponent=1.000000,bInvert=False))
-+AxisConfig=(AxisKeyName="OculusTouch_Right_Thumbstick",AxisProperties=(DeadZone=0.000000,Sensitivity=1.000000,Exponent=1.000000,bInvert=False))
-+AxisConfig=(AxisKeyName="OculusTouch_Right_FaceButton1",AxisProperties=(DeadZone=0.000000,Sensitivity=1.000000,Exponent=1.000000,bInvert=False))
-+AxisConfig=(AxisKeyName="OculusTouch_Right_Trigger",AxisProperties=(DeadZone=0.000000,Sensitivity=1.000000,Exponent=1.000000,bInvert=False))
-+AxisConfig=(AxisKeyName="OculusTouch_Right_FaceButton2",AxisProperties=(DeadZone=0.000000,Sensitivity=1.000000,Exponent=1.000000,bInvert=False))
-+AxisConfig=(AxisKeyName="OculusTouch_Right_IndexPointing",AxisProperties=(DeadZone=0.000000,Sensitivity=1.000000,Exponent=1.000000,bInvert=False))
-+AxisConfig=(AxisKeyName="OculusTouch_Right_ThumbUp",AxisProperties=(DeadZone=0.000000,Sensitivity=1.000000,Exponent=1.000000,bInvert=False))
-+AxisConfig=(AxisKeyName="OculusTouchpad_Touchpad_X",AxisProperties=(DeadZone=0.000000,Sensitivity=1.000000,Exponent=1.000000,bInvert=False))
-+AxisConfig=(AxisKeyName="OculusTouchpad_Touchpad_Y",AxisProperties=(DeadZone=0.000000,Sensitivity=1.000000,Exponent=1.000000,bInvert=False))
-+AxisConfig=(AxisKeyName="SteamVR_Knuckles_Left_HandGrip",AxisProperties=(DeadZone=0.000000,Sensitivity=1.000000,Exponent=1.000000,bInvert=False))
-+AxisConfig=(AxisKeyName="SteamVR_Knuckles_Left_IndexGrip",AxisProperties=(DeadZone=0.000000,Sensitivity=1.000000,Exponent=1.000000,bInvert=False))
-+AxisConfig=(AxisKeyName="SteamVR_Knuckles_Left_MiddleGrip",AxisProperties=(DeadZone=0.000000,Sensitivity=1.000000,Exponent=1.000000,bInvert=False))
-+AxisConfig=(AxisKeyName="SteamVR_Knuckles_Left_RingGrip",AxisProperties=(DeadZone=0.000000,Sensitivity=1.000000,Exponent=1.000000,bInvert=False))
-+AxisConfig=(AxisKeyName="SteamVR_Knuckles_Left_PinkyGrip",AxisProperties=(DeadZone=0.000000,Sensitivity=1.000000,Exponent=1.000000,bInvert=False))
-+AxisConfig=(AxisKeyName="SteamVR_Knuckles_Right_HandGrip",AxisProperties=(DeadZone=0.000000,Sensitivity=1.000000,Exponent=1.000000,bInvert=False))
-+AxisConfig=(AxisKeyName="SteamVR_Knuckles_Right_IndexGrip",AxisProperties=(DeadZone=0.000000,Sensitivity=1.000000,Exponent=1.000000,bInvert=False))
-+AxisConfig=(AxisKeyName="SteamVR_Knuckles_Right_MiddleGrip",AxisProperties=(DeadZone=0.000000,Sensitivity=1.000000,Exponent=1.000000,bInvert=False))
-+AxisConfig=(AxisKeyName="SteamVR_Knuckles_Right_RingGrip",AxisProperties=(DeadZone=0.000000,Sensitivity=1.000000,Exponent=1.000000,bInvert=False))
-+AxisConfig=(AxisKeyName="SteamVR_Knuckles_Right_PinkyGrip",AxisProperties=(DeadZone=0.000000,Sensitivity=1.000000,Exponent=1.000000,bInvert=False))
-+AxisConfig=(AxisKeyName="Vive_Left_Trigger_Axis",AxisProperties=(DeadZone=0.000000,Sensitivity=1.000000,Exponent=1.000000,bInvert=False))
-+AxisConfig=(AxisKeyName="Vive_Left_Trackpad_X",AxisProperties=(DeadZone=0.000000,Sensitivity=1.000000,Exponent=1.000000,bInvert=False))
-+AxisConfig=(AxisKeyName="Vive_Left_Trackpad_Y",AxisProperties=(DeadZone=0.000000,Sensitivity=1.000000,Exponent=1.000000,bInvert=False))
-+AxisConfig=(AxisKeyName="Vive_Right_Trigger_Axis",AxisProperties=(DeadZone=0.000000,Sensitivity=1.000000,Exponent=1.000000,bInvert=False))
-+AxisConfig=(AxisKeyName="Vive_Right_Trackpad_X",AxisProperties=(DeadZone=0.000000,Sensitivity=1.000000,Exponent=1.000000,bInvert=False))
-+AxisConfig=(AxisKeyName="Vive_Right_Trackpad_Y",AxisProperties=(DeadZone=0.000000,Sensitivity=1.000000,Exponent=1.000000,bInvert=False))
-+AxisConfig=(AxisKeyName="MixedReality_Left_Trigger_Axis",AxisProperties=(DeadZone=0.000000,Sensitivity=1.000000,Exponent=1.000000,bInvert=False))
-+AxisConfig=(AxisKeyName="MixedReality_Left_Thumbstick_X",AxisProperties=(DeadZone=0.000000,Sensitivity=1.000000,Exponent=1.000000,bInvert=False))
-+AxisConfig=(AxisKeyName="MixedReality_Left_Thumbstick_Y",AxisProperties=(DeadZone=0.000000,Sensitivity=1.000000,Exponent=1.000000,bInvert=False))
-+AxisConfig=(AxisKeyName="MixedReality_Left_Trackpad_X",AxisProperties=(DeadZone=0.000000,Sensitivity=1.000000,Exponent=1.000000,bInvert=False))
-+AxisConfig=(AxisKeyName="MixedReality_Left_Trackpad_Y",AxisProperties=(DeadZone=0.000000,Sensitivity=1.000000,Exponent=1.000000,bInvert=False))
-+AxisConfig=(AxisKeyName="MixedReality_Right_Trigger_Axis",AxisProperties=(DeadZone=0.000000,Sensitivity=1.000000,Exponent=1.000000,bInvert=False))
-+AxisConfig=(AxisKeyName="MixedReality_Right_Thumbstick_X",AxisProperties=(DeadZone=0.000000,Sensitivity=1.000000,Exponent=1.000000,bInvert=False))
-+AxisConfig=(AxisKeyName="MixedReality_Right_Thumbstick_Y",AxisProperties=(DeadZone=0.000000,Sensitivity=1.000000,Exponent=1.000000,bInvert=False))
-+AxisConfig=(AxisKeyName="MixedReality_Right_Trackpad_X",AxisProperties=(DeadZone=0.000000,Sensitivity=1.000000,Exponent=1.000000,bInvert=False))
-+AxisConfig=(AxisKeyName="MixedReality_Right_Trackpad_Y",AxisProperties=(DeadZone=0.000000,Sensitivity=1.000000,Exponent=1.000000,bInvert=False))
-+AxisConfig=(AxisKeyName="OculusGo_Left_Trackpad_X",AxisProperties=(DeadZone=0.000000,Sensitivity=1.000000,Exponent=1.000000,bInvert=False))
-+AxisConfig=(AxisKeyName="OculusGo_Left_Trackpad_Y",AxisProperties=(DeadZone=0.000000,Sensitivity=1.000000,Exponent=1.000000,bInvert=False))
-+AxisConfig=(AxisKeyName="OculusGo_Right_Trackpad_X",AxisProperties=(DeadZone=0.000000,Sensitivity=1.000000,Exponent=1.000000,bInvert=False))
-+AxisConfig=(AxisKeyName="OculusGo_Right_Trackpad_Y",AxisProperties=(DeadZone=0.000000,Sensitivity=1.000000,Exponent=1.000000,bInvert=False))
-+AxisConfig=(AxisKeyName="OculusTouch_Left_Grip_Axis",AxisProperties=(DeadZone=0.000000,Sensitivity=1.000000,Exponent=1.000000,bInvert=False))
-+AxisConfig=(AxisKeyName="OculusTouch_Left_Trigger_Axis",AxisProperties=(DeadZone=0.000000,Sensitivity=1.000000,Exponent=1.000000,bInvert=False))
-+AxisConfig=(AxisKeyName="OculusTouch_Left_Thumbstick_X",AxisProperties=(DeadZone=0.000000,Sensitivity=1.000000,Exponent=1.000000,bInvert=False))
-+AxisConfig=(AxisKeyName="OculusTouch_Left_Thumbstick_Y",AxisProperties=(DeadZone=0.000000,Sensitivity=1.000000,Exponent=1.000000,bInvert=False))
-+AxisConfig=(AxisKeyName="OculusTouch_Right_Grip_Axis",AxisProperties=(DeadZone=0.000000,Sensitivity=1.000000,Exponent=1.000000,bInvert=False))
-+AxisConfig=(AxisKeyName="OculusTouch_Right_Trigger_Axis",AxisProperties=(DeadZone=0.000000,Sensitivity=1.000000,Exponent=1.000000,bInvert=False))
-+AxisConfig=(AxisKeyName="OculusTouch_Right_Thumbstick_X",AxisProperties=(DeadZone=0.000000,Sensitivity=1.000000,Exponent=1.000000,bInvert=False))
-+AxisConfig=(AxisKeyName="OculusTouch_Right_Thumbstick_Y",AxisProperties=(DeadZone=0.000000,Sensitivity=1.000000,Exponent=1.000000,bInvert=False))
-+AxisConfig=(AxisKeyName="ValveIndex_Left_Grip_Axis",AxisProperties=(DeadZone=0.000000,Sensitivity=1.000000,Exponent=1.000000,bInvert=False))
-+AxisConfig=(AxisKeyName="ValveIndex_Left_Grip_Force",AxisProperties=(DeadZone=0.000000,Sensitivity=1.000000,Exponent=1.000000,bInvert=False))
-+AxisConfig=(AxisKeyName="ValveIndex_Left_Trigger_Axis",AxisProperties=(DeadZone=0.000000,Sensitivity=1.000000,Exponent=1.000000,bInvert=False))
-+AxisConfig=(AxisKeyName="ValveIndex_Left_Thumbstick_X",AxisProperties=(DeadZone=0.000000,Sensitivity=1.000000,Exponent=1.000000,bInvert=False))
-+AxisConfig=(AxisKeyName="ValveIndex_Left_Thumbstick_Y",AxisProperties=(DeadZone=0.000000,Sensitivity=1.000000,Exponent=1.000000,bInvert=False))
-+AxisConfig=(AxisKeyName="ValveIndex_Left_Trackpad_X",AxisProperties=(DeadZone=0.000000,Sensitivity=1.000000,Exponent=1.000000,bInvert=False))
-+AxisConfig=(AxisKeyName="ValveIndex_Left_Trackpad_Y",AxisProperties=(DeadZone=0.000000,Sensitivity=1.000000,Exponent=1.000000,bInvert=False))
-+AxisConfig=(AxisKeyName="ValveIndex_Left_Trackpad_Force",AxisProperties=(DeadZone=0.000000,Sensitivity=1.000000,Exponent=1.000000,bInvert=False))
-+AxisConfig=(AxisKeyName="ValveIndex_Left_Trackpad_Touch",AxisProperties=(DeadZone=0.000000,Sensitivity=1.000000,Exponent=1.000000,bInvert=False))
-+AxisConfig=(AxisKeyName="ValveIndex_Right_Grip_Axis",AxisProperties=(DeadZone=0.000000,Sensitivity=1.000000,Exponent=1.000000,bInvert=False))
-+AxisConfig=(AxisKeyName="ValveIndex_Right_Grip_Force",AxisProperties=(DeadZone=0.000000,Sensitivity=1.000000,Exponent=1.000000,bInvert=False))
-+AxisConfig=(AxisKeyName="ValveIndex_Right_Trigger_Axis",AxisProperties=(DeadZone=0.000000,Sensitivity=1.000000,Exponent=1.000000,bInvert=False))
-+AxisConfig=(AxisKeyName="ValveIndex_Right_Thumbstick_X",AxisProperties=(DeadZone=0.000000,Sensitivity=1.000000,Exponent=1.000000,bInvert=False))
-+AxisConfig=(AxisKeyName="ValveIndex_Right_Thumbstick_Y",AxisProperties=(DeadZone=0.000000,Sensitivity=1.000000,Exponent=1.000000,bInvert=False))
-+AxisConfig=(AxisKeyName="ValveIndex_Right_Trackpad_X",AxisProperties=(DeadZone=0.000000,Sensitivity=1.000000,Exponent=1.000000,bInvert=False))
-+AxisConfig=(AxisKeyName="ValveIndex_Right_Trackpad_Y",AxisProperties=(DeadZone=0.000000,Sensitivity=1.000000,Exponent=1.000000,bInvert=False))
-+AxisConfig=(AxisKeyName="ValveIndex_Right_Trackpad_Force",AxisProperties=(DeadZone=0.000000,Sensitivity=1.000000,Exponent=1.000000,bInvert=False))
 bAltEnterTogglesFullscreen=True
 bF11TogglesFullscreen=True
 bUseMouseForTouch=False
 bEnableMouseSmoothing=True
 bEnableFOVScaling=True
 bCaptureMouseOnLaunch=True
+bEnableLegacyInputScales=True
 bAlwaysShowTouchInterface=False
 bShowConsoleOnFourFingerTap=True
 bEnableGestureRecognizer=False
@@ -110,19 +37,6 @@ DoubleClickTime=0.200000
 +ActionMappings=(ActionName="Jump",bShift=False,bCtrl=False,bAlt=False,bCmd=False,Key=Gamepad_FaceButton_Bottom)
 +ActionMappings=(ActionName="Fire",bShift=False,bCtrl=False,bAlt=False,bCmd=False,Key=LeftMouseButton)
 +ActionMappings=(ActionName="Fire",bShift=False,bCtrl=False,bAlt=False,bCmd=False,Key=Gamepad_RightTrigger)
-+ActionMappings=(ActionName="ResetVR",bShift=False,bCtrl=False,bAlt=False,bCmd=False,Key=R)
-+ActionMappings=(ActionName="ResetVR",bShift=False,bCtrl=False,bAlt=False,bCmd=False,Key=Vive_Left_Grip_Click)
-+ActionMappings=(ActionName="Fire",bShift=False,bCtrl=False,bAlt=False,bCmd=False,Key=Vive_Right_Trigger_Click)
-+ActionMappings=(ActionName="Jump",bShift=False,bCtrl=False,bAlt=False,bCmd=False,Key=Vive_Left_Trigger_Click)
-+ActionMappings=(ActionName="Jump",bShift=False,bCtrl=False,bAlt=False,bCmd=False,Key=MixedReality_Left_Trigger_Click)
-+ActionMappings=(ActionName="Jump",bShift=False,bCtrl=False,bAlt=False,bCmd=False,Key=OculusTouch_Left_Trigger_Click)
-+ActionMappings=(ActionName="Jump",bShift=False,bCtrl=False,bAlt=False,bCmd=False,Key=ValveIndex_Left_Trigger_Click)
-+ActionMappings=(ActionName="Fire",bShift=False,bCtrl=False,bAlt=False,bCmd=False,Key=MixedReality_Right_Trigger_Click)
-+ActionMappings=(ActionName="Fire",bShift=False,bCtrl=False,bAlt=False,bCmd=False,Key=OculusTouch_Right_Trigger_Click)
-+ActionMappings=(ActionName="Fire",bShift=False,bCtrl=False,bAlt=False,bCmd=False,Key=ValveIndex_Right_Trigger_Click)
-+ActionMappings=(ActionName="ResetVR",bShift=False,bCtrl=False,bAlt=False,bCmd=False,Key=MixedReality_Left_Thumbstick_Click)
-+ActionMappings=(ActionName="ResetVR",bShift=False,bCtrl=False,bAlt=False,bCmd=False,Key=OculusTouch_Left_Thumbstick_Click)
-+ActionMappings=(ActionName="ResetVR",bShift=False,bCtrl=False,bAlt=False,bCmd=False,Key=ValveIndex_Left_Thumbstick_Click)
 +AxisMappings=(AxisName="MoveForward",Scale=1.000000,Key=W)
 +AxisMappings=(AxisName="MoveForward",Scale=-1.000000,Key=S)
 +AxisMappings=(AxisName="MoveForward",Scale=1.000000,Key=Up)
@@ -137,17 +51,9 @@ DoubleClickTime=0.200000
 +AxisMappings=(AxisName="Turn",Scale=1.000000,Key=MouseX)
 +AxisMappings=(AxisName="LookUpRate",Scale=1.000000,Key=Gamepad_RightY)
 +AxisMappings=(AxisName="LookUp",Scale=-1.000000,Key=MouseY)
-+AxisMappings=(AxisName="MoveForward",Scale=1.000000,Key=Vive_Left_Trackpad_Y)
-+AxisMappings=(AxisName="MoveForward",Scale=1.000000,Key=MixedReality_Left_Thumbstick_Y)
-+AxisMappings=(AxisName="MoveForward",Scale=1.000000,Key=OculusGo_Left_Trackpad_Y)
-+AxisMappings=(AxisName="MoveForward",Scale=1.000000,Key=OculusTouch_Left_Thumbstick_Y)
-+AxisMappings=(AxisName="MoveForward",Scale=1.000000,Key=ValveIndex_Left_Thumbstick_Y)
-+AxisMappings=(AxisName="MoveRight",Scale=1.000000,Key=Vive_Left_Trackpad_X)
-+AxisMappings=(AxisName="MoveRight",Scale=1.000000,Key=MixedReality_Left_Thumbstick_X)
-+AxisMappings=(AxisName="MoveRight",Scale=1.000000,Key=OculusGo_Left_Trackpad_X)
-+AxisMappings=(AxisName="MoveRight",Scale=1.000000,Key=OculusTouch_Left_Thumbstick_Y)
-+AxisMappings=(AxisName="MoveRight",Scale=1.000000,Key=ValveIndex_Left_Thumbstick_Y)
+DefaultPlayerInputClass=/Script/Engine.PlayerInput
+DefaultInputComponentClass=/Script/Engine.InputComponent
 DefaultTouchInterface=/Engine/MobileResources/HUD/DefaultVirtualJoysticks.DefaultVirtualJoysticks
+-ConsoleKeys=Tilde
 +ConsoleKeys=Tilde
 
-
diff --git a/Templates/TP_FirstPerson/Config/config.ini b/Templates/TP_FirstPerson/Config/config.ini
new file mode 100644
--- /dev/null
+++ b/Templates/TP_FirstPerson/Config/config.ini
@@ -0,0 +1,35 @@
+[InputSettings]
++AxisMappings=(AxisName="MoveForward",Scale=1.000000,Key=W)
++AxisMappings=(AxisName="MoveForward",Scale=-1.000000,Key=S)
++AxisMappings=(AxisName="MoveForward",Scale=1.000000,Key=Up)
++AxisMappings=(AxisName="MoveForward",Scale=-1.000000,Key=Down)
++AxisMappings=(AxisName="MoveForward",Scale=1.000000,Key=Gamepad_LeftY)
++AxisMappings=(AxisName="MoveRight",Scale=-1.000000,Key=A)
++AxisMappings=(AxisName="MoveRight",Scale=1.000000,Key=D)
++AxisMappings=(AxisName="MoveRight",Scale=1.000000,Key=Gamepad_LeftX)
++AxisMappings=(AxisName="TurnRate",Scale=1.000000,Key=Gamepad_RightX)
++AxisMappings=(AxisName="TurnRate",Scale=-1.000000,Key=Left)
++AxisMappings=(AxisName="TurnRate",Scale=1.000000,Key=Right)
++AxisMappings=(AxisName="Turn",Scale=1.000000,Key=MouseX)
++AxisMappings=(AxisName="LookUpRate",Scale=1.000000,Key=Gamepad_RightY)
++AxisMappings=(AxisName="LookUp",Scale=-1.000000,Key=MouseY)
++ActionMappings=(ActionName="Jump",bShift=False,bCtrl=False,bAlt=False,bCmd=False,Key=SpaceBar)
++ActionMappings=(ActionName="Jump",bShift=False,bCtrl=False,bAlt=False,bCmd=False,Key=Gamepad_FaceButton_Bottom)
++ActionMappings=(ActionName="Fire",bShift=False,bCtrl=False,bAlt=False,bCmd=False,Key=LeftMouseButton)
++ActionMappings=(ActionName="Fire",bShift=False,bCtrl=False,bAlt=False,bCmd=False,Key=Gamepad_RightTrigger)
+
+[Redirects]
+GameName=TP_FirstPerson
+
+[AdditionalFilesToAdd]
++Files=Templates/TP_FirstPerson/Content/*.*
++Files=Templates/TP_FirstPerson/Source/TP_FirstPerson/TP_FirstPerson.h
++Files=Templates/TP_FirstPerson/Source/TP_FirstPerson/TP_FirstPerson.cpp
++Files=Templates/TP_FirstPerson/Source/TP_FirstPerson/TP_FirstPersonCharacter.cpp
++Files=Templates/TP_FirstPerson/Source/TP_FirstPerson/TP_FirstPersonCharacter.h
++Files=Templates/TP_FirstPerson/Source/TP_FirstPerson/TP_FirstPersonGameMode.cpp
++Files=Templates/TP_FirstPerson/Source/TP_FirstPerson/TP_FirstPersonGameMode.h
++Files=Templates/TP_FirstPerson/Source/TP_FirstPerson/TP_FirstPersonHUD.cpp
++Files=Templates/TP_FirstPerson/Source/TP_FirstPerson/TP_FirstPersonHUD.h
++Files=Templates/TP_FirstPerson/Source/TP_FirstPerson/TP_FirstPersonProjectile.cpp
++Files=Templates/TP_FirstPerson/Source/TP_FirstPerson/TP_FirstPersonProjectile.h
diff --git a/Templates/TP_FirstPerson/Source/TP_FirstPerson/TP_FirstPersonCharacter.cpp b/Templates/TP_FirstPerson/Source/TP_FirstPerson/TP_FirstPersonCharacter.cpp
--- a/Templates/TP_FirstPerson/Source/TP_FirstPerson/TP_FirstPersonCharacter.cpp
+++ b/Templates/TP_FirstPerson/Source/TP_FirstPerson/TP_FirstPersonCharacter.cpp
@@ -7,10 +7,7 @@
 #include "Components/CapsuleComponent.h"
 #include "Components/InputComponent.h"
 #include "GameFramework/InputSettings.h"
-#include "HeadMountedDisplayFunctionLibrary.h"
 #include "Kismet/GameplayStatics.h"
-#include "MotionControllerComponent.h"
-#include "XRMotionControllerBase.h" // for FXRMotionControllerBase::RightHandSourceId
 
 DEFINE_LOG_CATEGORY_STATIC(LogFPChar, Warning, All);
 
@@ -43,7 +40,7 @@ ATP_FirstPersonCharacter::ATP_FirstPersonCharacter()
 
 	// Create a gun mesh component
 	FP_Gun = CreateDefaultSubobject<USkeletalMeshComponent>(TEXT("FP_Gun"));
-	FP_Gun->SetOnlyOwnerSee(false);			// otherwise won't be visible in the multiplayer
+	FP_Gun->SetOnlyOwnerSee(false);			// otherwise won't be visible in multiplayer
 	FP_Gun->bCastDynamicShadow = false;
 	FP_Gun->CastShadow = false;
 	// FP_Gun->SetupAttachment(Mesh1P, TEXT("GripPoint"));
@@ -56,32 +53,8 @@ ATP_FirstPersonCharacter::ATP_FirstPersonCharacter()
 	// Default offset from the character location for projectiles to spawn
 	GunOffset = FVector(100.0f, 0.0f, 10.0f);
 
-	// Note: The ProjectileClass and the skeletal mesh/anim blueprints for Mesh1P, FP_Gun, and VR_Gun 
-	// are set in the derived blueprint asset named MyCharacter to avoid direct content references in C++.
-
-	// Create VR Controllers.
-	R_MotionController = CreateDefaultSubobject<UMotionControllerComponent>(TEXT("R_MotionController"));
-	R_MotionController->MotionSource = FXRMotionControllerBase::RightHandSourceId;
-	R_MotionController->SetupAttachment(RootComponent);
-	L_MotionController = CreateDefaultSubobject<UMotionControllerComponent>(TEXT("L_MotionController"));
-	L_MotionController->SetupAttachment(RootComponent);
-
-	// Create a gun and attach it to the right-hand VR controller.
-	// Create a gun mesh component
-	VR_Gun = CreateDefaultSubobject<USkeletalMeshComponent>(TEXT("VR_Gun"));
-	VR_Gun->SetOnlyOwnerSee(false);			// otherwise won't be visible in the multiplayer
-	VR_Gun->bCastDynamicShadow = false;
-	VR_Gun->CastShadow = false;
-	VR_Gun->SetupAttachment(R_MotionController);
-	VR_Gun->SetRelativeRotation(FRotator(0.0f, -90.0f, 0.0f));
-
-	VR_MuzzleLocation = CreateDefaultSubobject<USceneComponent>(TEXT("VR_MuzzleLocation"));
-	VR_MuzzleLocation->SetupAttachment(VR_Gun);
-	VR_MuzzleLocation->SetRelativeLocation(FVector(0.000004, 53.999992, 10.000000));
-	VR_MuzzleLocation->SetRelativeRotation(FRotator(0.0f, 90.0f, 0.0f));		// Counteract the rotation of the VR gun model.
-
-	// Uncomment the following line to turn motion controllers on by default:
-	//bUsingMotionControllers = true;
+	// Note: The ProjectileClass and the skeletal mesh/anim blueprints for Mesh1P and FP_Gun
+	// are set in the derived blueprint asset named FirstPersonCharacter to avoid direct content references in C++.
 }
 
 void ATP_FirstPersonCharacter::BeginPlay()
@@ -91,26 +64,13 @@ void ATP_FirstPersonCharacter::BeginPlay()
 
 	//Attach gun mesh component to Skeleton, doing it here because the skeleton is not yet created in the constructor
 	FP_Gun->AttachToComponent(Mesh1P, FAttachmentTransformRules(EAttachmentRule::SnapToTarget, true), TEXT("GripPoint"));
-
-	// Show or hide the two versions of the gun based on whether or not we're using motion controllers.
-	if (bUsingMotionControllers)
-	{
-		VR_Gun->SetHiddenInGame(false, true);
-		Mesh1P->SetHiddenInGame(true, true);
-	}
-	else
-	{
-		VR_Gun->SetHiddenInGame(true, true);
-		Mesh1P->SetHiddenInGame(false, true);
-	}
 }
 
-//////////////////////////////////////////////////////////////////////////
-// Input
+//////////////////////////////////////////////////////////////////////////// Input
 
 void ATP_FirstPersonCharacter::SetupPlayerInputComponent(class UInputComponent* PlayerInputComponent)
 {
-	// set up gameplay key bindings
+	// Set up gameplay key bindings
 	check(PlayerInputComponent);
 
 	// Bind jump events
@@ -123,8 +83,6 @@ void ATP_FirstPersonCharacter::SetupPlayerInputComponent(class UInputComponent*
 	// Enable touchscreen input
 	EnableTouchscreenMovement(PlayerInputComponent);
 
-	PlayerInputComponent->BindAction("ResetVR", IE_Pressed, this, &ATP_FirstPersonCharacter::OnResetVR);
-
 	// Bind movement events
 	PlayerInputComponent->BindAxis("MoveForward", this, &ATP_FirstPersonCharacter::MoveForward);
 	PlayerInputComponent->BindAxis("MoveRight", this, &ATP_FirstPersonCharacter::MoveRight);
@@ -140,41 +98,32 @@ void ATP_FirstPersonCharacter::SetupPlayerInputComponent(class UInputComponent*
 
 void ATP_FirstPersonCharacter::OnFire()
 {
-	// try and fire a projectile
+	// Try and fire a projectile
 	if (ProjectileClass != nullptr)
 	{
 		UWorld* const World = GetWorld();
 		if (World != nullptr)
 		{
-			if (bUsingMotionControllers)
-			{
-				const FRotator SpawnRotation = VR_MuzzleLocation->GetComponentRotation();
-				const FVector SpawnLocation = VR_MuzzleLocation->GetComponentLocation();
-				World->SpawnActor<ATP_FirstPersonProjectile>(ProjectileClass, SpawnLocation, SpawnRotation);
-			}
-			else
-			{
-				const FRotator SpawnRotation = GetControlRotation();
-				// MuzzleOffset is in camera space, so transform it to world space before offsetting from the character location to find the final muzzle position
-				const FVector SpawnLocation = ((FP_MuzzleLocation != nullptr) ? FP_MuzzleLocation->GetComponentLocation() : GetActorLocation()) + SpawnRotation.RotateVector(GunOffset);
-
-				//Set Spawn Collision Handling Override
-				FActorSpawnParameters ActorSpawnParams;
-				ActorSpawnParams.SpawnCollisionHandlingOverride = ESpawnActorCollisionHandlingMethod::AdjustIfPossibleButDontSpawnIfColliding;
-
-				// spawn the projectile at the muzzle
-				World->SpawnActor<ATP_FirstPersonProjectile>(ProjectileClass, SpawnLocation, SpawnRotation, ActorSpawnParams);
-			}
+			const FRotator SpawnRotation = GetControlRotation();
+			// MuzzleOffset is in camera space, so transform it to world space before offsetting from the character location to find the final muzzle position
+			const FVector SpawnLocation = ((FP_MuzzleLocation != nullptr) ? FP_MuzzleLocation->GetComponentLocation() : GetActorLocation()) + SpawnRotation.RotateVector(GunOffset);
+
+			//Set Spawn Collision Handling Override
+			FActorSpawnParameters ActorSpawnParams;
+			ActorSpawnParams.SpawnCollisionHandlingOverride = ESpawnActorCollisionHandlingMethod::AdjustIfPossibleButDontSpawnIfColliding;
+
+			// Spawn the projectile at the muzzle
+			World->SpawnActor<ATP_FirstPersonProjectile>(ProjectileClass, SpawnLocation, SpawnRotation, ActorSpawnParams);
 		}
 	}
 
-	// try and play the sound if specified
+	// Try and play the sound if specified
 	if (FireSound != nullptr)
 	{
 		UGameplayStatics::PlaySoundAtLocation(this, FireSound, GetActorLocation());
 	}
 
-	// try and play a firing animation if specified
+	// Try and play a firing animation if specified
 	if (FireAnimation != nullptr)
 	{
 		// Get the animation object for the arms mesh
@@ -186,11 +135,6 @@ void ATP_FirstPersonCharacter::OnFire()
 	}
 }
 
-void ATP_FirstPersonCharacter::OnResetVR()
-{
-	UHeadMountedDisplayFunctionLibrary::ResetOrientationAndPosition();
-}
-
 void ATP_FirstPersonCharacter::BeginTouch(const ETouchIndex::Type FingerIndex, const FVector Location)
 {
 	if (TouchItem.bIsPressed == true)
diff --git a/Templates/TP_FirstPerson/Source/TP_FirstPerson/TP_FirstPersonCharacter.h b/Templates/TP_FirstPerson/Source/TP_FirstPerson/TP_FirstPersonCharacter.h
--- a/Templates/TP_FirstPerson/Source/TP_FirstPerson/TP_FirstPersonCharacter.h
+++ b/Templates/TP_FirstPerson/Source/TP_FirstPerson/TP_FirstPersonCharacter.h
@@ -10,7 +10,6 @@ class UInputComponent;
 class USkeletalMeshComponent;
 class USceneComponent;
 class UCameraComponent;
-class UMotionControllerComponent;
 class UAnimMontage;
 class USoundBase;
 
@@ -31,26 +30,10 @@ class ATP_FirstPersonCharacter : public ACharacter
 	UPROPERTY(VisibleDefaultsOnly, Category = Mesh)
 	USceneComponent* FP_MuzzleLocation;
 
-	/** Gun mesh: VR view (attached to the VR controller directly, no arm, just the actual gun) */
-	UPROPERTY(VisibleDefaultsOnly, Category = Mesh)
-	USkeletalMeshComponent* VR_Gun;
-
-	/** Location on VR gun mesh where projectiles should spawn. */
-	UPROPERTY(VisibleDefaultsOnly, Category = Mesh)
-	USceneComponent* VR_MuzzleLocation;
-
 	/** First person camera */
 	UPROPERTY(VisibleAnywhere, BlueprintReadOnly, Category = Camera, meta = (AllowPrivateAccess = "true"))
 	UCameraComponent* FirstPersonCameraComponent;
 
-	/** Motion controller (right hand) */
-	UPROPERTY(VisibleAnywhere, BlueprintReadOnly, meta = (AllowPrivateAccess = "true"))
-	UMotionControllerComponent* R_MotionController;
-
-	/** Motion controller (left hand) */
-	UPROPERTY(VisibleAnywhere, BlueprintReadOnly, meta = (AllowPrivateAccess = "true"))
-	UMotionControllerComponent* L_MotionController;
-
 public:
 	ATP_FirstPersonCharacter();
 
@@ -82,22 +65,15 @@ public:
 	UPROPERTY(EditAnywhere, BlueprintReadWrite, Category = Gameplay)
 	UAnimMontage* FireAnimation;
 
-	/** Whether to use motion controller location for aiming. */
-	UPROPERTY(EditAnywhere, BlueprintReadWrite, Category = Gameplay)
-	uint8 bUsingMotionControllers : 1;
-
 protected:
 	
 	/** Fires a projectile. */
 	void OnFire();
 
-	/** Resets HMD orientation and position in VR. */
-	void OnResetVR();
-
 	/** Handles moving forward/backward */
 	void MoveForward(float Val);
 
-	/** Handles stafing movement, left and right */
+	/** Handles strafing movement, left and right */
 	void MoveRight(float Val);
 
 	/**
diff --git a/Templates/TP_FirstPerson/contents.txt b/Templates/TP_FirstPerson/contents.txt
new file mode 100644
--- /dev/null
+++ b/Templates/TP_FirstPerson/contents.txt
@@ -0,0 +1,4 @@
+../../../Templates/TP_FirstPerson/manifest.json
+../../../Templates/TP_FirstPerson/Config/config.ini
+../../../Templates/TP_FirstPerson/Media/TP_FirstPerson.png
+../../../Templates/TP_FirstPerson/Media/TP_FirstPerson_Preview.png
diff --git a/Templates/TP_FirstPerson/manifest.json b/Templates/TP_FirstPerson/manifest.json
new file mode 100644
--- /dev/null
+++ b/Templates/TP_FirstPerson/manifest.json
@@ -0,0 +1,85 @@
+{
+	"Version": 1,
+	"Name":
+	[
+		{
+			"Language": "en",
+			"Text": "First Person"
+		},
+		{
+			"Language": "ko",
+			"Text": "일인칭"
+		},
+		{
+			"Language": "zh-Hans",
+			"Text": "第一人称游戏"
+		},
+		{
+			"Language": "ja",
+			"Text": "ファーストパーソン"
+		}
+	],
+	"Description":
+	[
+		{
+			"Language": "en",
+			"Text": "The first person template features a player character (represented by a pair of arms) which is viewed from first person perspective. The character can be moved around the level using keyboard, controller or virtual joystick on a touch device. Additionally the player can look around using mouse, controller or virtual joystick on a touch device. The character is also equipped with a gun that, using mouse,  controller or virtual joystick on a touch device will fire a simple sphere projectile that will affect some physics objects in the level, whilst bounce off the arena walls."
+		},
+		{
+			"Language": "ko",
+			"Text": "일인칭 템플릿은 일인칭 시점에서 보이는 (양팔로 표현되는) 플레이어 캐릭터를 특징으로 합니다. 캐릭터 이동은 키보드, 컨트롤러, 터치 디바이스의 가상 조이스틱으로 가능합니다. 추가적으로 플레이어는 마우스, 컨트롤러, 터치 디바이스의 가상 조이스틱을 사용하여 주변을 둘러볼 수도 있습니다. 캐릭터에는 총이 장착되어 있어 마우스, 컨트롤러, 터치 디바이스의 가상 조이스틱을 사용하면 간단한 발사체가 발사되어, 아레나 벽에 튕기면서 레벨에 있는 몇몇 물리 오브젝트에 영향을 줍니다."
+		},
+		{
+			"Language": "zh-Hans",
+			"Text": "第一人称模板拥有一个第一人称的玩家角色（通过一双手臂呈现）。您可以通过键盘、手柄或触摸式设备上的虚拟摇杆控制角色在关卡中四处移动。此外，您还可以使用鼠标、手柄或触摸式设备上的虚拟摇杆转动角色视角，观察四周。角色还装备了枪支，可使用鼠标、手柄或触摸式设备上的虚拟摇杆发射一颗球形子弹，并影响关卡中的一些物理对象，同时从竞技场的墙壁上弹开。"
+		},
+		{
+			"Language": "ja",
+			"Text": "ファーストパーソンテンプレートには、一人称視点で表示されるプレイヤーキャラクタ（腕として表示されています）がいます。キーボード、コントローラー、タッチデバイスでのバーチャルジョイスティックでこのキャラクタを移動させることができます。追加して、プレイヤーはマウス、コントローラー、タッチデバイスでのバーチャルジョイスティックで周りを見回すことができます。さらにキャラクタは銃を装備して、マウス、コントローラー、タッチデバイスでのバーチャルジョイスティックを使ってシンプルな球体の弾を発射できます。この弾は物理オブジェクトに影響し、アリーナの壁で反射します。"
+		}
+	],
+	"AssetTypes":
+	[
+		{
+			"Language": "en",
+			"Text": "Animation Sequence, Animation Montage, Sound Wave, Skeletal Mesh, Animation Blueprint, Skeleton, Static Mesh, Texture."
+		}
+	],	
+	"SearchTags":
+	[
+		{
+			"Language": "en",
+			"Text": "First Person,Shoot,Fire,Aim"
+		}
+	],
+	"ClassTypes": "Character, GameMode, HUD, Actor, Character, GameMode, HUD, Actor",
+	"Category": "CodeFeature",
+	"FocusAsset": "/Game/FirstPersonCPP/Blueprints/FirstPersonCharacter.FirstPersonCharacter",
+	"Thumbnail": "TP_FirstPerson.png",
+	"Screenshots":
+	[
+		"TP_FirstPerson_Preview.png"
+	],
+	"AdditionalFeaturePacks": [
+		{
+			"MountName": "Geometry",
+			"DetailLevels": [
+				"Standard",
+				"High"
+			]
+		},
+		{
+			"MountName": "FPWeapon",
+			"DetailLevels": [
+				"Standard"
+			]
+		},
+		{
+			"MountName": "FirstPerson",
+			"DetailLevels":
+			[
+				"Standard"
+			]
+		}
+	]
+}
diff --git a/Templates/TP_FirstPersonBP/Config/DefaultInput.ini b/Templates/TP_FirstPersonBP/Config/DefaultInput.ini
--- a/Templates/TP_FirstPersonBP/Config/DefaultInput.ini
+++ b/Templates/TP_FirstPersonBP/Config/DefaultInput.ini
@@ -1,103 +1,30 @@
-
 [/Script/Engine.InputSettings]
+-AxisConfig=(AxisKeyName="Gamepad_LeftX",AxisProperties=(DeadZone=0.25,Exponent=1.f,Sensitivity=1.f))
+-AxisConfig=(AxisKeyName="Gamepad_LeftY",AxisProperties=(DeadZone=0.25,Exponent=1.f,Sensitivity=1.f))
+-AxisConfig=(AxisKeyName="Gamepad_RightX",AxisProperties=(DeadZone=0.25,Exponent=1.f,Sensitivity=1.f))
+-AxisConfig=(AxisKeyName="Gamepad_RightY",AxisProperties=(DeadZone=0.25,Exponent=1.f,Sensitivity=1.f))
+-AxisConfig=(AxisKeyName="MouseX",AxisProperties=(DeadZone=0.f,Exponent=1.f,Sensitivity=0.07f))
+-AxisConfig=(AxisKeyName="MouseY",AxisProperties=(DeadZone=0.f,Exponent=1.f,Sensitivity=0.07f))
+-AxisConfig=(AxisKeyName="Mouse2D",AxisProperties=(DeadZone=0.f,Exponent=1.f,Sensitivity=0.07f))
++AxisConfig=(AxisKeyName="Mouse2D",AxisProperties=(DeadZone=0.000000,Sensitivity=0.070000,Exponent=1.000000,bInvert=False))
 +AxisConfig=(AxisKeyName="Gamepad_LeftX",AxisProperties=(DeadZone=0.250000,Sensitivity=1.000000,Exponent=1.000000,bInvert=False))
 +AxisConfig=(AxisKeyName="Gamepad_LeftY",AxisProperties=(DeadZone=0.250000,Sensitivity=1.000000,Exponent=1.000000,bInvert=False))
 +AxisConfig=(AxisKeyName="Gamepad_RightX",AxisProperties=(DeadZone=0.250000,Sensitivity=1.000000,Exponent=1.000000,bInvert=False))
 +AxisConfig=(AxisKeyName="Gamepad_RightY",AxisProperties=(DeadZone=0.250000,Sensitivity=1.000000,Exponent=1.000000,bInvert=False))
 +AxisConfig=(AxisKeyName="MouseX",AxisProperties=(DeadZone=0.000000,Sensitivity=0.070000,Exponent=1.000000,bInvert=False))
 +AxisConfig=(AxisKeyName="MouseY",AxisProperties=(DeadZone=0.000000,Sensitivity=0.070000,Exponent=1.000000,bInvert=False))
-+AxisConfig=(AxisKeyName="MotionController_Left_Thumbstick_Z",AxisProperties=(DeadZone=0.000000,Sensitivity=1.000000,Exponent=1.000000,bInvert=False))
-+AxisConfig=(AxisKeyName="MotionController_Right_Thumbstick_Z",AxisProperties=(DeadZone=0.000000,Sensitivity=1.000000,Exponent=1.000000,bInvert=False))
 +AxisConfig=(AxisKeyName="MouseWheelAxis",AxisProperties=(DeadZone=0.000000,Sensitivity=1.000000,Exponent=1.000000,bInvert=False))
 +AxisConfig=(AxisKeyName="Gamepad_LeftTriggerAxis",AxisProperties=(DeadZone=0.000000,Sensitivity=1.000000,Exponent=1.000000,bInvert=False))
 +AxisConfig=(AxisKeyName="Gamepad_RightTriggerAxis",AxisProperties=(DeadZone=0.000000,Sensitivity=1.000000,Exponent=1.000000,bInvert=False))
-+AxisConfig=(AxisKeyName="MotionController_Left_Thumbstick_X",AxisProperties=(DeadZone=0.000000,Sensitivity=1.000000,Exponent=1.000000,bInvert=False))
-+AxisConfig=(AxisKeyName="MotionController_Left_Thumbstick_Y",AxisProperties=(DeadZone=0.000000,Sensitivity=1.000000,Exponent=1.000000,bInvert=False))
-+AxisConfig=(AxisKeyName="MotionController_Left_TriggerAxis",AxisProperties=(DeadZone=0.000000,Sensitivity=1.000000,Exponent=1.000000,bInvert=False))
-+AxisConfig=(AxisKeyName="MotionController_Left_Grip1Axis",AxisProperties=(DeadZone=0.000000,Sensitivity=1.000000,Exponent=1.000000,bInvert=False))
-+AxisConfig=(AxisKeyName="MotionController_Left_Grip2Axis",AxisProperties=(DeadZone=0.000000,Sensitivity=1.000000,Exponent=1.000000,bInvert=False))
-+AxisConfig=(AxisKeyName="MotionController_Right_Thumbstick_X",AxisProperties=(DeadZone=0.000000,Sensitivity=1.000000,Exponent=1.000000,bInvert=False))
-+AxisConfig=(AxisKeyName="MotionController_Right_Thumbstick_Y",AxisProperties=(DeadZone=0.000000,Sensitivity=1.000000,Exponent=1.000000,bInvert=False))
-+AxisConfig=(AxisKeyName="MotionController_Right_TriggerAxis",AxisProperties=(DeadZone=0.000000,Sensitivity=1.000000,Exponent=1.000000,bInvert=False))
-+AxisConfig=(AxisKeyName="MotionController_Right_Grip1Axis",AxisProperties=(DeadZone=0.000000,Sensitivity=1.000000,Exponent=1.000000,bInvert=False))
-+AxisConfig=(AxisKeyName="MotionController_Right_Grip2Axis",AxisProperties=(DeadZone=0.000000,Sensitivity=1.000000,Exponent=1.000000,bInvert=False))
 +AxisConfig=(AxisKeyName="Gamepad_Special_Left_X",AxisProperties=(DeadZone=0.000000,Sensitivity=1.000000,Exponent=1.000000,bInvert=False))
 +AxisConfig=(AxisKeyName="Gamepad_Special_Left_Y",AxisProperties=(DeadZone=0.000000,Sensitivity=1.000000,Exponent=1.000000,bInvert=False))
-+AxisConfig=(AxisKeyName="OculusTouch_Left_Thumbstick",AxisProperties=(DeadZone=0.000000,Sensitivity=1.000000,Exponent=1.000000,bInvert=False))
-+AxisConfig=(AxisKeyName="OculusTouch_Left_FaceButton1",AxisProperties=(DeadZone=0.000000,Sensitivity=1.000000,Exponent=1.000000,bInvert=False))
-+AxisConfig=(AxisKeyName="OculusTouch_Left_Trigger",AxisProperties=(DeadZone=0.000000,Sensitivity=1.000000,Exponent=1.000000,bInvert=False))
-+AxisConfig=(AxisKeyName="OculusTouch_Left_FaceButton2",AxisProperties=(DeadZone=0.000000,Sensitivity=1.000000,Exponent=1.000000,bInvert=False))
-+AxisConfig=(AxisKeyName="OculusTouch_Left_IndexPointing",AxisProperties=(DeadZone=0.000000,Sensitivity=1.000000,Exponent=1.000000,bInvert=False))
-+AxisConfig=(AxisKeyName="OculusTouch_Left_ThumbUp",AxisProperties=(DeadZone=0.000000,Sensitivity=1.000000,Exponent=1.000000,bInvert=False))
-+AxisConfig=(AxisKeyName="OculusTouch_Right_Thumbstick",AxisProperties=(DeadZone=0.000000,Sensitivity=1.000000,Exponent=1.000000,bInvert=False))
-+AxisConfig=(AxisKeyName="OculusTouch_Right_FaceButton1",AxisProperties=(DeadZone=0.000000,Sensitivity=1.000000,Exponent=1.000000,bInvert=False))
-+AxisConfig=(AxisKeyName="OculusTouch_Right_Trigger",AxisProperties=(DeadZone=0.000000,Sensitivity=1.000000,Exponent=1.000000,bInvert=False))
-+AxisConfig=(AxisKeyName="OculusTouch_Right_FaceButton2",AxisProperties=(DeadZone=0.000000,Sensitivity=1.000000,Exponent=1.000000,bInvert=False))
-+AxisConfig=(AxisKeyName="OculusTouch_Right_IndexPointing",AxisProperties=(DeadZone=0.000000,Sensitivity=1.000000,Exponent=1.000000,bInvert=False))
-+AxisConfig=(AxisKeyName="OculusTouch_Right_ThumbUp",AxisProperties=(DeadZone=0.000000,Sensitivity=1.000000,Exponent=1.000000,bInvert=False))
-+AxisConfig=(AxisKeyName="OculusTouchpad_Touchpad_X",AxisProperties=(DeadZone=0.000000,Sensitivity=1.000000,Exponent=1.000000,bInvert=False))
-+AxisConfig=(AxisKeyName="OculusTouchpad_Touchpad_Y",AxisProperties=(DeadZone=0.000000,Sensitivity=1.000000,Exponent=1.000000,bInvert=False))
-+AxisConfig=(AxisKeyName="SteamVR_Knuckles_Left_HandGrip",AxisProperties=(DeadZone=0.000000,Sensitivity=1.000000,Exponent=1.000000,bInvert=False))
-+AxisConfig=(AxisKeyName="SteamVR_Knuckles_Left_IndexGrip",AxisProperties=(DeadZone=0.000000,Sensitivity=1.000000,Exponent=1.000000,bInvert=False))
-+AxisConfig=(AxisKeyName="SteamVR_Knuckles_Left_MiddleGrip",AxisProperties=(DeadZone=0.000000,Sensitivity=1.000000,Exponent=1.000000,bInvert=False))
-+AxisConfig=(AxisKeyName="SteamVR_Knuckles_Left_RingGrip",AxisProperties=(DeadZone=0.000000,Sensitivity=1.000000,Exponent=1.000000,bInvert=False))
-+AxisConfig=(AxisKeyName="SteamVR_Knuckles_Left_PinkyGrip",AxisProperties=(DeadZone=0.000000,Sensitivity=1.000000,Exponent=1.000000,bInvert=False))
-+AxisConfig=(AxisKeyName="SteamVR_Knuckles_Right_HandGrip",AxisProperties=(DeadZone=0.000000,Sensitivity=1.000000,Exponent=1.000000,bInvert=False))
-+AxisConfig=(AxisKeyName="SteamVR_Knuckles_Right_IndexGrip",AxisProperties=(DeadZone=0.000000,Sensitivity=1.000000,Exponent=1.000000,bInvert=False))
-+AxisConfig=(AxisKeyName="SteamVR_Knuckles_Right_MiddleGrip",AxisProperties=(DeadZone=0.000000,Sensitivity=1.000000,Exponent=1.000000,bInvert=False))
-+AxisConfig=(AxisKeyName="SteamVR_Knuckles_Right_RingGrip",AxisProperties=(DeadZone=0.000000,Sensitivity=1.000000,Exponent=1.000000,bInvert=False))
-+AxisConfig=(AxisKeyName="SteamVR_Knuckles_Right_PinkyGrip",AxisProperties=(DeadZone=0.000000,Sensitivity=1.000000,Exponent=1.000000,bInvert=False))
-+AxisConfig=(AxisKeyName="Vive_Left_Trigger_Axis",AxisProperties=(DeadZone=0.000000,Sensitivity=1.000000,Exponent=1.000000,bInvert=False))
-+AxisConfig=(AxisKeyName="Vive_Left_Trackpad_X",AxisProperties=(DeadZone=0.000000,Sensitivity=1.000000,Exponent=1.000000,bInvert=False))
-+AxisConfig=(AxisKeyName="Vive_Left_Trackpad_Y",AxisProperties=(DeadZone=0.000000,Sensitivity=1.000000,Exponent=1.000000,bInvert=False))
-+AxisConfig=(AxisKeyName="Vive_Right_Trigger_Axis",AxisProperties=(DeadZone=0.000000,Sensitivity=1.000000,Exponent=1.000000,bInvert=False))
-+AxisConfig=(AxisKeyName="Vive_Right_Trackpad_X",AxisProperties=(DeadZone=0.000000,Sensitivity=1.000000,Exponent=1.000000,bInvert=False))
-+AxisConfig=(AxisKeyName="Vive_Right_Trackpad_Y",AxisProperties=(DeadZone=0.000000,Sensitivity=1.000000,Exponent=1.000000,bInvert=False))
-+AxisConfig=(AxisKeyName="MixedReality_Left_Trigger_Axis",AxisProperties=(DeadZone=0.000000,Sensitivity=1.000000,Exponent=1.000000,bInvert=False))
-+AxisConfig=(AxisKeyName="MixedReality_Left_Thumbstick_X",AxisProperties=(DeadZone=0.000000,Sensitivity=1.000000,Exponent=1.000000,bInvert=False))
-+AxisConfig=(AxisKeyName="MixedReality_Left_Thumbstick_Y",AxisProperties=(DeadZone=0.000000,Sensitivity=1.000000,Exponent=1.000000,bInvert=False))
-+AxisConfig=(AxisKeyName="MixedReality_Left_Trackpad_X",AxisProperties=(DeadZone=0.000000,Sensitivity=1.000000,Exponent=1.000000,bInvert=False))
-+AxisConfig=(AxisKeyName="MixedReality_Left_Trackpad_Y",AxisProperties=(DeadZone=0.000000,Sensitivity=1.000000,Exponent=1.000000,bInvert=False))
-+AxisConfig=(AxisKeyName="MixedReality_Right_Trigger_Axis",AxisProperties=(DeadZone=0.000000,Sensitivity=1.000000,Exponent=1.000000,bInvert=False))
-+AxisConfig=(AxisKeyName="MixedReality_Right_Thumbstick_X",AxisProperties=(DeadZone=0.000000,Sensitivity=1.000000,Exponent=1.000000,bInvert=False))
-+AxisConfig=(AxisKeyName="MixedReality_Right_Thumbstick_Y",AxisProperties=(DeadZone=0.000000,Sensitivity=1.000000,Exponent=1.000000,bInvert=False))
-+AxisConfig=(AxisKeyName="MixedReality_Right_Trackpad_X",AxisProperties=(DeadZone=0.000000,Sensitivity=1.000000,Exponent=1.000000,bInvert=False))
-+AxisConfig=(AxisKeyName="MixedReality_Right_Trackpad_Y",AxisProperties=(DeadZone=0.000000,Sensitivity=1.000000,Exponent=1.000000,bInvert=False))
-+AxisConfig=(AxisKeyName="OculusGo_Left_Trackpad_X",AxisProperties=(DeadZone=0.000000,Sensitivity=1.000000,Exponent=1.000000,bInvert=False))
-+AxisConfig=(AxisKeyName="OculusGo_Left_Trackpad_Y",AxisProperties=(DeadZone=0.000000,Sensitivity=1.000000,Exponent=1.000000,bInvert=False))
-+AxisConfig=(AxisKeyName="OculusGo_Right_Trackpad_X",AxisProperties=(DeadZone=0.000000,Sensitivity=1.000000,Exponent=1.000000,bInvert=False))
-+AxisConfig=(AxisKeyName="OculusGo_Right_Trackpad_Y",AxisProperties=(DeadZone=0.000000,Sensitivity=1.000000,Exponent=1.000000,bInvert=False))
-+AxisConfig=(AxisKeyName="OculusTouch_Left_Grip_Axis",AxisProperties=(DeadZone=0.000000,Sensitivity=1.000000,Exponent=1.000000,bInvert=False))
-+AxisConfig=(AxisKeyName="OculusTouch_Left_Trigger_Axis",AxisProperties=(DeadZone=0.000000,Sensitivity=1.000000,Exponent=1.000000,bInvert=False))
-+AxisConfig=(AxisKeyName="OculusTouch_Left_Thumbstick_X",AxisProperties=(DeadZone=0.000000,Sensitivity=1.000000,Exponent=1.000000,bInvert=False))
-+AxisConfig=(AxisKeyName="OculusTouch_Left_Thumbstick_Y",AxisProperties=(DeadZone=0.000000,Sensitivity=1.000000,Exponent=1.000000,bInvert=False))
-+AxisConfig=(AxisKeyName="OculusTouch_Right_Grip_Axis",AxisProperties=(DeadZone=0.000000,Sensitivity=1.000000,Exponent=1.000000,bInvert=False))
-+AxisConfig=(AxisKeyName="OculusTouch_Right_Trigger_Axis",AxisProperties=(DeadZone=0.000000,Sensitivity=1.000000,Exponent=1.000000,bInvert=False))
-+AxisConfig=(AxisKeyName="OculusTouch_Right_Thumbstick_X",AxisProperties=(DeadZone=0.000000,Sensitivity=1.000000,Exponent=1.000000,bInvert=False))
-+AxisConfig=(AxisKeyName="OculusTouch_Right_Thumbstick_Y",AxisProperties=(DeadZone=0.000000,Sensitivity=1.000000,Exponent=1.000000,bInvert=False))
-+AxisConfig=(AxisKeyName="ValveIndex_Left_Grip_Axis",AxisProperties=(DeadZone=0.000000,Sensitivity=1.000000,Exponent=1.000000,bInvert=False))
-+AxisConfig=(AxisKeyName="ValveIndex_Left_Grip_Force",AxisProperties=(DeadZone=0.000000,Sensitivity=1.000000,Exponent=1.000000,bInvert=False))
-+AxisConfig=(AxisKeyName="ValveIndex_Left_Trigger_Axis",AxisProperties=(DeadZone=0.000000,Sensitivity=1.000000,Exponent=1.000000,bInvert=False))
-+AxisConfig=(AxisKeyName="ValveIndex_Left_Thumbstick_X",AxisProperties=(DeadZone=0.000000,Sensitivity=1.000000,Exponent=1.000000,bInvert=False))
-+AxisConfig=(AxisKeyName="ValveIndex_Left_Thumbstick_Y",AxisProperties=(DeadZone=0.000000,Sensitivity=1.000000,Exponent=1.000000,bInvert=False))
-+AxisConfig=(AxisKeyName="ValveIndex_Left_Trackpad_X",AxisProperties=(DeadZone=0.000000,Sensitivity=1.000000,Exponent=1.000000,bInvert=False))
-+AxisConfig=(AxisKeyName="ValveIndex_Left_Trackpad_Y",AxisProperties=(DeadZone=0.000000,Sensitivity=1.000000,Exponent=1.000000,bInvert=False))
-+AxisConfig=(AxisKeyName="ValveIndex_Left_Trackpad_Force",AxisProperties=(DeadZone=0.000000,Sensitivity=1.000000,Exponent=1.000000,bInvert=False))
-+AxisConfig=(AxisKeyName="ValveIndex_Left_Trackpad_Touch",AxisProperties=(DeadZone=0.000000,Sensitivity=1.000000,Exponent=1.000000,bInvert=False))
-+AxisConfig=(AxisKeyName="ValveIndex_Right_Grip_Axis",AxisProperties=(DeadZone=0.000000,Sensitivity=1.000000,Exponent=1.000000,bInvert=False))
-+AxisConfig=(AxisKeyName="ValveIndex_Right_Grip_Force",AxisProperties=(DeadZone=0.000000,Sensitivity=1.000000,Exponent=1.000000,bInvert=False))
-+AxisConfig=(AxisKeyName="ValveIndex_Right_Trigger_Axis",AxisProperties=(DeadZone=0.000000,Sensitivity=1.000000,Exponent=1.000000,bInvert=False))
-+AxisConfig=(AxisKeyName="ValveIndex_Right_Thumbstick_X",AxisProperties=(DeadZone=0.000000,Sensitivity=1.000000,Exponent=1.000000,bInvert=False))
-+AxisConfig=(AxisKeyName="ValveIndex_Right_Thumbstick_Y",AxisProperties=(DeadZone=0.000000,Sensitivity=1.000000,Exponent=1.000000,bInvert=False))
-+AxisConfig=(AxisKeyName="ValveIndex_Right_Trackpad_X",AxisProperties=(DeadZone=0.000000,Sensitivity=1.000000,Exponent=1.000000,bInvert=False))
-+AxisConfig=(AxisKeyName="ValveIndex_Right_Trackpad_Y",AxisProperties=(DeadZone=0.000000,Sensitivity=1.000000,Exponent=1.000000,bInvert=False))
-+AxisConfig=(AxisKeyName="ValveIndex_Right_Trackpad_Force",AxisProperties=(DeadZone=0.000000,Sensitivity=1.000000,Exponent=1.000000,bInvert=False))
 bAltEnterTogglesFullscreen=True
 bF11TogglesFullscreen=True
 bUseMouseForTouch=False
 bEnableMouseSmoothing=True
 bEnableFOVScaling=True
 bCaptureMouseOnLaunch=True
+bEnableLegacyInputScales=True
 bAlwaysShowTouchInterface=False
 bShowConsoleOnFourFingerTap=True
 bEnableGestureRecognizer=False
@@ -110,19 +37,6 @@ DoubleClickTime=0.200000
 +ActionMappings=(ActionName="Jump",bShift=False,bCtrl=False,bAlt=False,bCmd=False,Key=Gamepad_FaceButton_Bottom)
 +ActionMappings=(ActionName="Fire",bShift=False,bCtrl=False,bAlt=False,bCmd=False,Key=LeftMouseButton)
 +ActionMappings=(ActionName="Fire",bShift=False,bCtrl=False,bAlt=False,bCmd=False,Key=Gamepad_RightTrigger)
-+ActionMappings=(ActionName="ResetVR",bShift=False,bCtrl=False,bAlt=False,bCmd=False,Key=R)
-+ActionMappings=(ActionName="ResetVR",bShift=False,bCtrl=False,bAlt=False,bCmd=False,Key=Vive_Left_Grip_Click)
-+ActionMappings=(ActionName="Fire",bShift=False,bCtrl=False,bAlt=False,bCmd=False,Key=Vive_Right_Trigger_Click)
-+ActionMappings=(ActionName="Jump",bShift=False,bCtrl=False,bAlt=False,bCmd=False,Key=Vive_Left_Trigger_Click)
-+ActionMappings=(ActionName="Jump",bShift=False,bCtrl=False,bAlt=False,bCmd=False,Key=MixedReality_Left_Trigger_Click)
-+ActionMappings=(ActionName="Jump",bShift=False,bCtrl=False,bAlt=False,bCmd=False,Key=OculusTouch_Left_Trigger_Click)
-+ActionMappings=(ActionName="Jump",bShift=False,bCtrl=False,bAlt=False,bCmd=False,Key=ValveIndex_Left_Trigger_Click)
-+ActionMappings=(ActionName="Fire",bShift=False,bCtrl=False,bAlt=False,bCmd=False,Key=MixedReality_Right_Trigger_Click)
-+ActionMappings=(ActionName="Fire",bShift=False,bCtrl=False,bAlt=False,bCmd=False,Key=OculusTouch_Right_Trigger_Click)
-+ActionMappings=(ActionName="Fire",bShift=False,bCtrl=False,bAlt=False,bCmd=False,Key=ValveIndex_Right_Trigger_Click)
-+ActionMappings=(ActionName="ResetVR",bShift=False,bCtrl=False,bAlt=False,bCmd=False,Key=MixedReality_Left_Thumbstick_Click)
-+ActionMappings=(ActionName="ResetVR",bShift=False,bCtrl=False,bAlt=False,bCmd=False,Key=OculusTouch_Left_Thumbstick_Click)
-+ActionMappings=(ActionName="ResetVR",bShift=False,bCtrl=False,bAlt=False,bCmd=False,Key=ValveIndex_Left_Thumbstick_Click)
 +AxisMappings=(AxisName="MoveForward",Scale=1.000000,Key=W)
 +AxisMappings=(AxisName="MoveForward",Scale=-1.000000,Key=S)
 +AxisMappings=(AxisName="MoveForward",Scale=1.000000,Key=Up)
@@ -137,17 +51,9 @@ DoubleClickTime=0.200000
 +AxisMappings=(AxisName="Turn",Scale=1.000000,Key=MouseX)
 +AxisMappings=(AxisName="LookUpRate",Scale=1.000000,Key=Gamepad_RightY)
 +AxisMappings=(AxisName="LookUp",Scale=-1.000000,Key=MouseY)
-+AxisMappings=(AxisName="MoveForward",Scale=1.000000,Key=Vive_Left_Trackpad_Y)
-+AxisMappings=(AxisName="MoveForward",Scale=1.000000,Key=MixedReality_Left_Thumbstick_Y)
-+AxisMappings=(AxisName="MoveForward",Scale=1.000000,Key=OculusGo_Left_Trackpad_Y)
-+AxisMappings=(AxisName="MoveForward",Scale=1.000000,Key=OculusTouch_Left_Thumbstick_Y)
-+AxisMappings=(AxisName="MoveForward",Scale=1.000000,Key=ValveIndex_Left_Thumbstick_Y)
-+AxisMappings=(AxisName="MoveRight",Scale=1.000000,Key=Vive_Left_Trackpad_X)
-+AxisMappings=(AxisName="MoveRight",Scale=1.000000,Key=MixedReality_Left_Thumbstick_X)
-+AxisMappings=(AxisName="MoveRight",Scale=1.000000,Key=OculusGo_Left_Trackpad_X)
-+AxisMappings=(AxisName="MoveRight",Scale=1.000000,Key=OculusTouch_Left_Thumbstick_Y)
-+AxisMappings=(AxisName="MoveRight",Scale=1.000000,Key=ValveIndex_Left_Thumbstick_Y)
+DefaultPlayerInputClass=/Script/Engine.PlayerInput
+DefaultInputComponentClass=/Script/Engine.InputComponent
 DefaultTouchInterface=/Engine/MobileResources/HUD/DefaultVirtualJoysticks.DefaultVirtualJoysticks
+-ConsoleKeys=Tilde
 +ConsoleKeys=Tilde
 
-
diff --git a/Templates/TP_FirstPersonBP/Config/config.ini b/Templates/TP_FirstPersonBP/Config/config.ini
new file mode 100644
--- /dev/null
+++ b/Templates/TP_FirstPersonBP/Config/config.ini
@@ -0,0 +1,26 @@
+[InputSettings]
++AxisMappings=(AxisName="MoveForward",Scale=1.000000,Key=W)
++AxisMappings=(AxisName="MoveForward",Scale=-1.000000,Key=S)
++AxisMappings=(AxisName="MoveForward",Scale=1.000000,Key=Up)
++AxisMappings=(AxisName="MoveForward",Scale=-1.000000,Key=Down)
++AxisMappings=(AxisName="MoveForward",Scale=1.000000,Key=Gamepad_LeftY)
++AxisMappings=(AxisName="MoveRight",Scale=-1.000000,Key=A)
++AxisMappings=(AxisName="MoveRight",Scale=1.000000,Key=D)
++AxisMappings=(AxisName="MoveRight",Scale=1.000000,Key=Gamepad_LeftX)
++AxisMappings=(AxisName="TurnRate",Scale=1.000000,Key=Gamepad_RightX)
++AxisMappings=(AxisName="TurnRate",Scale=-1.000000,Key=Left)
++AxisMappings=(AxisName="TurnRate",Scale=1.000000,Key=Right)
++AxisMappings=(AxisName="Turn",Scale=1.000000,Key=MouseX)
++AxisMappings=(AxisName="LookUpRate",Scale=1.000000,Key=Gamepad_RightY)
++AxisMappings=(AxisName="LookUp",Scale=-1.000000,Key=MouseY)
++ActionMappings=(ActionName="Jump",bShift=False,bCtrl=False,bAlt=False,bCmd=False,Key=SpaceBar)
++ActionMappings=(ActionName="Jump",bShift=False,bCtrl=False,bAlt=False,bCmd=False,Key=Gamepad_FaceButton_Bottom)
++ActionMappings=(ActionName="Fire",bShift=False,bCtrl=False,bAlt=False,bCmd=False,Key=LeftMouseButton)
++ActionMappings=(ActionName="Fire",bShift=False,bCtrl=False,bAlt=False,bCmd=False,Key=Gamepad_RightTrigger)
+
+[Redirects]
+GameName=TP_FirstPersonBP
+
+[AdditionalFilesToAdd]
++Files=Templates/TP_FirstPersonBP/Content/*.*
+
diff --git a/Templates/TP_FirstPersonBP/contents.txt b/Templates/TP_FirstPersonBP/contents.txt
new file mode 100644
--- /dev/null
+++ b/Templates/TP_FirstPersonBP/contents.txt
@@ -0,0 +1,4 @@
+../../../Templates/TP_FirstPersonBP/manifest.json
+../../../Templates/TP_FirstPersonBP/Config/config.ini
+../../../Templates/TP_FirstPersonBP/Media/TP_FirstPersonBP.png
+../../../Templates/TP_FirstPersonBP/Media/TP_FirstPersonBP_Preview.png
diff --git a/Templates/TP_FirstPersonBP/manifest.json b/Templates/TP_FirstPersonBP/manifest.json
new file mode 100644
--- /dev/null
+++ b/Templates/TP_FirstPersonBP/manifest.json
@@ -0,0 +1,85 @@
+{
+	"Version": 1,
+	"Name":
+	[
+		{
+			"Language": "en",
+			"Text": "First Person"
+		},
+		{
+			"Language": "ko",
+			"Text": "일인칭"
+		},
+		{
+			"Language": "zh-Hans",
+			"Text": "第一人称游戏"
+		},
+		{
+			"Language": "ja",
+			"Text": "ファーストパーソン"
+		}
+	],
+	"Description":
+	[
+		{
+			"Language": "en",
+			"Text": "The first person template features a player character (represented by a pair of arms) which is viewed from first person perspective. The character can be moved around the level using keyboard, controller or virtual joystick on a touch device. Additionally the player can look around using mouse, controller or virtual joystick on a touch device. The character is also equipped with a gun that, using mouse,  controller or virtual joystick on a touch device will fire a simple sphere projectile that will affect some physics objects in the level, whilst bounce off the arena walls."
+		},
+		{
+			"Language": "ko",
+			"Text": "일인칭 템플릿은 일인칭 시점에서 보이는 (양팔로 표현되는) 플레이어 캐릭터를 특징으로 합니다. 캐릭터 이동은 키보드, 컨트롤러, 터치 디바이스의 가상 조이스틱으로 가능합니다. 추가적으로 플레이어는 마우스, 컨트롤러, 터치 디바이스의 가상 조이스틱을 사용하여 주변을 둘러볼 수도 있습니다. 캐릭터에는 총이 장착되어 있어 마우스, 컨트롤러, 터치 디바이스의 가상 조이스틱을 사용하면 간단한 발사체가 발사되어, 아레나 벽에 튕기면서 레벨에 있는 몇몇 물리 오브젝트에 영향을 줍니다."
+		},
+		{
+			"Language": "zh-Hans",
+			"Text": "第一人称模板拥有一个第一人称的玩家角色（通过一双手臂呈现）。您可以通过键盘、手柄或触摸式设备上的虚拟摇杆控制角色在关卡中四处移动。此外，您还可以使用鼠标、手柄或触摸式设备上的虚拟摇杆转动角色视角，观察四周。角色还装备了枪支，可使用鼠标、手柄或触摸式设备上的虚拟摇杆发射一颗球形子弹，并影响关卡中的一些物理对象，同时从竞技场的墙壁上弹开。"
+		},
+		{
+			"Language": "ja",
+			"Text": "ファーストパーソンテンプレートには、一人称視点で表示されるプレイヤーキャラクタ（腕として表示されています）がいます。キーボード、コントローラー、タッチデバイスでのバーチャルジョイスティックでこのキャラクタを移動させることができます。追加して、プレイヤーはマウス、コントローラー、タッチデバイスでのバーチャルジョイスティックで周りを見回すことができます。さらにキャラクタは銃を装備して、マウス、コントローラー、タッチデバイスでのバーチャルジョイスティックを使ってシンプルな球体の弾を発射できます。この弾は物理オブジェクトに影響し、アリーナの壁で反射します。"
+		}
+	],
+	"AssetTypes":
+	[
+		{
+			"Language": "en",
+			"Text": "Animation Sequence, Animation Montage, Sound Wave, Skeletal Mesh, Animation Blueprint, Skeleton, Static Mesh, Texture."
+		}
+	],	
+	"SearchTags":
+	[
+		{
+			"Language": "en",
+			"Text": "First Person,Shoot,Fire,Aim"
+		}
+	],
+	"ClassTypes": "Character, GameMode, HUD, Actor, Character, GameMode, HUD, Actor",
+	"Category": "BlueprintFeature",
+	"FocusAsset": "/Game/FirstPersonBP/FirstPersonOverview.FirstPersonOverview",
+	"Thumbnail": "TP_FirstPersonBP.png",
+	"Screenshots":
+	[
+		"TP_FirstPersonBP_Preview.png"
+	],
+	"AdditionalFeaturePacks": [
+		{
+			"MountName": "Geometry",
+			"DetailLevels": [
+				"Standard",
+				"High"
+			]
+		},
+		{
+			"MountName": "FPWeapon",
+			"DetailLevels": [
+				"Standard"
+			]
+		},
+		{
+			"MountName": "FirstPerson",
+			"DetailLevels":
+			[
+				"Standard"
+			]
+		}
+	]
+}
