FirstPersonBP template refresh for 5.0:
FirstPersonCharacter no longer contains any hard references to the weapon and the weapon is now placed in the scene and picked up upon overlap. Logic and names are updated to reflect how the weapon could be any generic useable item.
Moved input logic to the Player Controller.
Updated FirstPersonExampleMap to use the new LevelPrototyping content.
Removed the HUD as it's not an accurate depiction of where the rifle shoots projectiles - this would require work not fit within the scope of a template.
Renamed TemplateResource packs with better names for a better user experience when working with our templates. This will temporarily break a couple of templates until we've updated them.

For more details refer to the Jiras' and the associated documents that covers the work.

This also resolves a bunch of old issues QA discovered within the template.

diff --git a/Templates/TP_FirstPerson/Config/DefaultEngine.ini b/Templates/TP_FirstPerson/Config/DefaultEngine.ini
--- a/Templates/TP_FirstPerson/Config/DefaultEngine.ini
+++ b/Templates/TP_FirstPerson/Config/DefaultEngine.ini
@@ -4,14 +4,14 @@
 +EditProfiles=(Name="Trigger",CustomResponses=((Channel=Projectile, Response=ECR_Ignore)))
 
 [/Script/EngineSettings.GameMapsSettings]
-EditorStartupMap=/Game/FirstPersonCPP/Maps/FirstPersonExampleMap
+EditorStartupMap=/Game/FirstPerson/Maps/FirstPersonMap.FirstPersonMap
 LocalMapOptions=
 TransitionMap=
 bUseSplitscreen=True
 TwoPlayerSplitscreenLayout=Horizontal
 ThreePlayerSplitscreenLayout=FavorTop
 GameInstanceClass=/Script/Engine.GameInstance
-GameDefaultMap=/Game/FirstPersonCPP/Maps/FirstPersonExampleMap
+GameDefaultMap=/Game/FirstPerson/Maps/FirstPersonMap.FirstPersonMap
 ServerDefaultMap=/Engine/Maps/Entry
 GlobalDefaultGameMode=/Script/TP_FirstPerson.TP_FirstPersonGameMode
 GlobalDefaultServerGameMode=None
@@ -19,3 +19,8 @@ GlobalDefaultServerGameMode=None
 [/Script/IOSRuntimeSettings.IOSRuntimeSettings]
 MinimumiOSVersion=IOS_13
 
+[/Script/Engine.RendererSettings]
+r.ReflectionMethod=1
+r.GenerateMeshDistanceFields=True
+r.DynamicGlobalIlluminationMethod=1
+
diff --git a/Templates/TP_FirstPerson/Config/TemplateDefs.ini b/Templates/TP_FirstPerson/Config/TemplateDefs.ini
--- a/Templates/TP_FirstPerson/Config/TemplateDefs.ini
+++ b/Templates/TP_FirstPerson/Config/TemplateDefs.ini
@@ -46,6 +46,6 @@ FilenameReplacements=(Extensions=("cpp","h","ini","cs"),From="%TEMPLATENAME%",To
 ReplacementsInFiles=(Extensions=("cpp","h","ini","cs"),From="%TEMPLATENAME_UPPERCASE%",To="%PROJECTNAME_UPPERCASE%",bCaseSensitive=true)
 ReplacementsInFiles=(Extensions=("cpp","h","ini","cs"),From="%TEMPLATENAME_LOWERCASE%",To="%PROJECTNAME_LOWERCASE%",bCaseSensitive=true)
 ReplacementsInFiles=(Extensions=("cpp","h","ini","cs"),From="%TEMPLATENAME%",To="%PROJECTNAME%",bCaseSensitive=false)
-SharedContentPacks=(MountName="FirstPerson",DetailLevels=("Standard"))
-SharedContentPacks=(MountName="Geometry",DetailLevels=("Standard","High"))
+SharedContentPacks=(MountName="FirstPersonArms",DetailLevels=("Standard"))
 SharedContentPacks=(MountName="FPWeapon",DetailLevels=("Standard"))
+SharedContentPacks=(MountName="LevelPrototyping",DetailLevels=("High"))
diff --git a/Templates/TP_FirstPerson/Config/config.ini b/Templates/TP_FirstPerson/Config/config.ini
--- a/Templates/TP_FirstPerson/Config/config.ini
+++ b/Templates/TP_FirstPerson/Config/config.ini
@@ -29,7 +29,5 @@ GameName=TP_FirstPerson
 +Files=Templates/TP_FirstPerson/Source/TP_FirstPerson/TP_FirstPersonCharacter.h
 +Files=Templates/TP_FirstPerson/Source/TP_FirstPerson/TP_FirstPersonGameMode.cpp
 +Files=Templates/TP_FirstPerson/Source/TP_FirstPerson/TP_FirstPersonGameMode.h
-+Files=Templates/TP_FirstPerson/Source/TP_FirstPerson/TP_FirstPersonHUD.cpp
-+Files=Templates/TP_FirstPerson/Source/TP_FirstPerson/TP_FirstPersonHUD.h
 +Files=Templates/TP_FirstPerson/Source/TP_FirstPerson/TP_FirstPersonProjectile.cpp
 +Files=Templates/TP_FirstPerson/Source/TP_FirstPerson/TP_FirstPersonProjectile.h
diff --git a/Templates/TP_FirstPerson/Source/TP_FirstPerson/TP_FirstPersonGameMode.cpp b/Templates/TP_FirstPerson/Source/TP_FirstPerson/TP_FirstPersonGameMode.cpp
--- a/Templates/TP_FirstPerson/Source/TP_FirstPerson/TP_FirstPersonGameMode.cpp
+++ b/Templates/TP_FirstPerson/Source/TP_FirstPerson/TP_FirstPersonGameMode.cpp
@@ -1,7 +1,6 @@
 // Copyright Epic Games, Inc. All Rights Reserved.
 
 #include "TP_FirstPersonGameMode.h"
-#include "TP_FirstPersonHUD.h"
 #include "TP_FirstPersonCharacter.h"
 #include "UObject/ConstructorHelpers.h"
 
@@ -9,9 +8,7 @@ ATP_FirstPersonGameMode::ATP_FirstPersonGameMode()
 	: Super()
 {
 	// set default pawn class to our Blueprinted character
-	static ConstructorHelpers::FClassFinder<APawn> PlayerPawnClassFinder(TEXT("/Game/FirstPersonCPP/Blueprints/FirstPersonCharacter"));
+	static ConstructorHelpers::FClassFinder<APawn> PlayerPawnClassFinder(TEXT("/Game/FirstPerson/Blueprints/FirstPersonCharacter"));
 	DefaultPawnClass = PlayerPawnClassFinder.Class;
 
-	// use our custom HUD class
-	HUDClass = ATP_FirstPersonHUD::StaticClass();
 }
diff --git a/Templates/TP_FirstPerson/Source/TP_FirstPerson/TP_FirstPersonHUD.cpp b/Templates/TP_FirstPerson/Source/TP_FirstPerson/TP_FirstPersonHUD.cpp
deleted file mode 100644
diff --git a/Templates/TP_FirstPerson/Source/TP_FirstPerson/TP_FirstPersonHUD.h b/Templates/TP_FirstPerson/Source/TP_FirstPerson/TP_FirstPersonHUD.h
deleted file mode 100644
diff --git a/Templates/TP_FirstPerson/manifest.json b/Templates/TP_FirstPerson/manifest.json
--- a/Templates/TP_FirstPerson/manifest.json
+++ b/Templates/TP_FirstPerson/manifest.json
@@ -54,28 +54,27 @@
 	],
 	"ClassTypes": "Character, GameMode, HUD, Actor, Character, GameMode, HUD, Actor",
 	"Category": "CodeFeature",
-	"FocusAsset": "/Game/FirstPersonCPP/Blueprints/FirstPersonCharacter.FirstPersonCharacter",
+	"FocusAsset": "/Game/FirstPerson/Blueprints/FirstPersonCharacter.FirstPersonCharacter",
 	"Thumbnail": "TP_FirstPerson.png",
 	"Screenshots":
 	[
 		"TP_FirstPerson_Preview.png"
 	],
 	"AdditionalFeaturePacks": [
-		{
-			"MountName": "Geometry",
+				{
+			"MountName": "FPWeapon",
 			"DetailLevels": [
-				"Standard",
-				"High"
+				"Standard"
 			]
 		},
 		{
-			"MountName": "FPWeapon",
+			"MountName": "LevelPrototyping",
 			"DetailLevels": [
-				"Standard"
+				"High"
 			]
 		},
 		{
-			"MountName": "FirstPerson",
+			"MountName": "FirstPersonArms",
 			"DetailLevels":
 			[
 				"Standard"
diff --git a/Templates/TP_FirstPersonBP/Config/DefaultEngine.ini b/Templates/TP_FirstPersonBP/Config/DefaultEngine.ini
--- a/Templates/TP_FirstPersonBP/Config/DefaultEngine.ini
+++ b/Templates/TP_FirstPersonBP/Config/DefaultEngine.ini
@@ -4,18 +4,23 @@
 +EditProfiles=(Name="Trigger",CustomResponses=((Channel=Projectile, Response=ECR_Ignore)))
 
 [/Script/EngineSettings.GameMapsSettings]
-EditorStartupMap=/Game/FirstPersonBP/Maps/FirstPersonExampleMap
+EditorStartupMap=/Game/FirstPerson/Maps/FirstPersonMap.FirstPersonMap
 LocalMapOptions=
 TransitionMap=
 bUseSplitscreen=True
 TwoPlayerSplitscreenLayout=Horizontal
 ThreePlayerSplitscreenLayout=FavorTop
 GameInstanceClass=/Script/Engine.GameInstance
-GameDefaultMap=/Game/FirstPersonBP/Maps/FirstPersonExampleMap
+GameDefaultMap=/Game/FirstPerson/Maps/FirstPersonMap.FirstPersonMap
 ServerDefaultMap=/Engine/Maps/Entry
-GlobalDefaultGameMode=/Game/FirstPersonBP/Blueprints/FirstPersonGameMode.FirstPersonGameMode_C
+GlobalDefaultGameMode=/Game/FirstPerson/Blueprints/BP_FirstPersonGameMode.BP_FirstPersonGameMode_C
 GlobalDefaultServerGameMode=None
 
 [/Script/IOSRuntimeSettings.IOSRuntimeSettings]
 MinimumiOSVersion=IOS_13
 
+[/Script/Engine.RendererSettings]
+r.ReflectionMethod=1
+r.GenerateMeshDistanceFields=True
+r.DynamicGlobalIlluminationMethod=1
+
diff --git a/Templates/TP_FirstPersonBP/Config/DefaultInput.ini b/Templates/TP_FirstPersonBP/Config/DefaultInput.ini
--- a/Templates/TP_FirstPersonBP/Config/DefaultInput.ini
+++ b/Templates/TP_FirstPersonBP/Config/DefaultInput.ini
@@ -18,6 +18,79 @@
 +AxisConfig=(AxisKeyName="Gamepad_RightTriggerAxis",AxisProperties=(DeadZone=0.000000,Sensitivity=1.000000,Exponent=1.000000,bInvert=False))
 +AxisConfig=(AxisKeyName="Gamepad_Special_Left_X",AxisProperties=(DeadZone=0.000000,Sensitivity=1.000000,Exponent=1.000000,bInvert=False))
 +AxisConfig=(AxisKeyName="Gamepad_Special_Left_Y",AxisProperties=(DeadZone=0.000000,Sensitivity=1.000000,Exponent=1.000000,bInvert=False))
++AxisConfig=(AxisKeyName="Vive_Left_Trigger_Axis",AxisProperties=(DeadZone=0.000000,Sensitivity=1.000000,Exponent=1.000000,bInvert=False))
++AxisConfig=(AxisKeyName="Vive_Left_Trackpad_X",AxisProperties=(DeadZone=0.000000,Sensitivity=1.000000,Exponent=1.000000,bInvert=False))
++AxisConfig=(AxisKeyName="Vive_Left_Trackpad_Y",AxisProperties=(DeadZone=0.000000,Sensitivity=1.000000,Exponent=1.000000,bInvert=False))
++AxisConfig=(AxisKeyName="Vive_Right_Trigger_Axis",AxisProperties=(DeadZone=0.000000,Sensitivity=1.000000,Exponent=1.000000,bInvert=False))
++AxisConfig=(AxisKeyName="Vive_Right_Trackpad_X",AxisProperties=(DeadZone=0.000000,Sensitivity=1.000000,Exponent=1.000000,bInvert=False))
++AxisConfig=(AxisKeyName="Vive_Right_Trackpad_Y",AxisProperties=(DeadZone=0.000000,Sensitivity=1.000000,Exponent=1.000000,bInvert=False))
++AxisConfig=(AxisKeyName="MixedReality_Left_Trigger_Axis",AxisProperties=(DeadZone=0.000000,Sensitivity=1.000000,Exponent=1.000000,bInvert=False))
++AxisConfig=(AxisKeyName="MixedReality_Left_Thumbstick_X",AxisProperties=(DeadZone=0.000000,Sensitivity=1.000000,Exponent=1.000000,bInvert=False))
++AxisConfig=(AxisKeyName="MixedReality_Left_Thumbstick_Y",AxisProperties=(DeadZone=0.000000,Sensitivity=1.000000,Exponent=1.000000,bInvert=False))
++AxisConfig=(AxisKeyName="MixedReality_Left_Trackpad_X",AxisProperties=(DeadZone=0.000000,Sensitivity=1.000000,Exponent=1.000000,bInvert=False))
++AxisConfig=(AxisKeyName="MixedReality_Left_Trackpad_Y",AxisProperties=(DeadZone=0.000000,Sensitivity=1.000000,Exponent=1.000000,bInvert=False))
++AxisConfig=(AxisKeyName="MixedReality_Right_Trigger_Axis",AxisProperties=(DeadZone=0.000000,Sensitivity=1.000000,Exponent=1.000000,bInvert=False))
++AxisConfig=(AxisKeyName="MixedReality_Right_Thumbstick_X",AxisProperties=(DeadZone=0.000000,Sensitivity=1.000000,Exponent=1.000000,bInvert=False))
++AxisConfig=(AxisKeyName="MixedReality_Right_Thumbstick_Y",AxisProperties=(DeadZone=0.000000,Sensitivity=1.000000,Exponent=1.000000,bInvert=False))
++AxisConfig=(AxisKeyName="MixedReality_Right_Trackpad_X",AxisProperties=(DeadZone=0.000000,Sensitivity=1.000000,Exponent=1.000000,bInvert=False))
++AxisConfig=(AxisKeyName="MixedReality_Right_Trackpad_Y",AxisProperties=(DeadZone=0.000000,Sensitivity=1.000000,Exponent=1.000000,bInvert=False))
++AxisConfig=(AxisKeyName="OculusTouch_Left_Grip_Axis",AxisProperties=(DeadZone=0.000000,Sensitivity=1.000000,Exponent=1.000000,bInvert=False))
++AxisConfig=(AxisKeyName="OculusTouch_Left_Trigger_Axis",AxisProperties=(DeadZone=0.000000,Sensitivity=1.000000,Exponent=1.000000,bInvert=False))
++AxisConfig=(AxisKeyName="OculusTouch_Left_Thumbstick_X",AxisProperties=(DeadZone=0.000000,Sensitivity=1.000000,Exponent=1.000000,bInvert=False))
++AxisConfig=(AxisKeyName="OculusTouch_Left_Thumbstick_Y",AxisProperties=(DeadZone=0.000000,Sensitivity=1.000000,Exponent=1.000000,bInvert=False))
++AxisConfig=(AxisKeyName="OculusTouch_Right_Grip_Axis",AxisProperties=(DeadZone=0.000000,Sensitivity=1.000000,Exponent=1.000000,bInvert=False))
++AxisConfig=(AxisKeyName="OculusTouch_Right_Trigger_Axis",AxisProperties=(DeadZone=0.000000,Sensitivity=1.000000,Exponent=1.000000,bInvert=False))
++AxisConfig=(AxisKeyName="OculusTouch_Right_Thumbstick_X",AxisProperties=(DeadZone=0.000000,Sensitivity=1.000000,Exponent=1.000000,bInvert=False))
++AxisConfig=(AxisKeyName="OculusTouch_Right_Thumbstick_Y",AxisProperties=(DeadZone=0.000000,Sensitivity=1.000000,Exponent=1.000000,bInvert=False))
++AxisConfig=(AxisKeyName="ValveIndex_Left_Grip_Axis",AxisProperties=(DeadZone=0.000000,Sensitivity=1.000000,Exponent=1.000000,bInvert=False))
++AxisConfig=(AxisKeyName="ValveIndex_Left_Grip_Force",AxisProperties=(DeadZone=0.000000,Sensitivity=1.000000,Exponent=1.000000,bInvert=False))
++AxisConfig=(AxisKeyName="ValveIndex_Left_Trigger_Axis",AxisProperties=(DeadZone=0.000000,Sensitivity=1.000000,Exponent=1.000000,bInvert=False))
++AxisConfig=(AxisKeyName="ValveIndex_Left_Thumbstick_X",AxisProperties=(DeadZone=0.000000,Sensitivity=1.000000,Exponent=1.000000,bInvert=False))
++AxisConfig=(AxisKeyName="ValveIndex_Left_Thumbstick_Y",AxisProperties=(DeadZone=0.000000,Sensitivity=1.000000,Exponent=1.000000,bInvert=False))
++AxisConfig=(AxisKeyName="ValveIndex_Left_Trackpad_X",AxisProperties=(DeadZone=0.000000,Sensitivity=1.000000,Exponent=1.000000,bInvert=False))
++AxisConfig=(AxisKeyName="ValveIndex_Left_Trackpad_Y",AxisProperties=(DeadZone=0.000000,Sensitivity=1.000000,Exponent=1.000000,bInvert=False))
++AxisConfig=(AxisKeyName="ValveIndex_Left_Trackpad_Force",AxisProperties=(DeadZone=0.000000,Sensitivity=1.000000,Exponent=1.000000,bInvert=False))
++AxisConfig=(AxisKeyName="ValveIndex_Left_Trackpad_Touch",AxisProperties=(DeadZone=0.000000,Sensitivity=1.000000,Exponent=1.000000,bInvert=False))
++AxisConfig=(AxisKeyName="ValveIndex_Right_Grip_Axis",AxisProperties=(DeadZone=0.000000,Sensitivity=1.000000,Exponent=1.000000,bInvert=False))
++AxisConfig=(AxisKeyName="ValveIndex_Right_Grip_Force",AxisProperties=(DeadZone=0.000000,Sensitivity=1.000000,Exponent=1.000000,bInvert=False))
++AxisConfig=(AxisKeyName="ValveIndex_Right_Trigger_Axis",AxisProperties=(DeadZone=0.000000,Sensitivity=1.000000,Exponent=1.000000,bInvert=False))
++AxisConfig=(AxisKeyName="ValveIndex_Right_Thumbstick_X",AxisProperties=(DeadZone=0.000000,Sensitivity=1.000000,Exponent=1.000000,bInvert=False))
++AxisConfig=(AxisKeyName="ValveIndex_Right_Thumbstick_Y",AxisProperties=(DeadZone=0.000000,Sensitivity=1.000000,Exponent=1.000000,bInvert=False))
++AxisConfig=(AxisKeyName="ValveIndex_Right_Trackpad_X",AxisProperties=(DeadZone=0.000000,Sensitivity=1.000000,Exponent=1.000000,bInvert=False))
++AxisConfig=(AxisKeyName="ValveIndex_Right_Trackpad_Y",AxisProperties=(DeadZone=0.000000,Sensitivity=1.000000,Exponent=1.000000,bInvert=False))
++AxisConfig=(AxisKeyName="ValveIndex_Right_Trackpad_Force",AxisProperties=(DeadZone=0.000000,Sensitivity=1.000000,Exponent=1.000000,bInvert=False))
++AxisConfig=(AxisKeyName="OculusTouch_Left_Thumbstick",AxisProperties=(DeadZone=0.000000,Sensitivity=1.000000,Exponent=1.000000,bInvert=False))
++AxisConfig=(AxisKeyName="OculusTouch_Left_FaceButton1",AxisProperties=(DeadZone=0.000000,Sensitivity=1.000000,Exponent=1.000000,bInvert=False))
++AxisConfig=(AxisKeyName="OculusTouch_Left_Trigger",AxisProperties=(DeadZone=0.000000,Sensitivity=1.000000,Exponent=1.000000,bInvert=False))
++AxisConfig=(AxisKeyName="OculusTouch_Left_FaceButton2",AxisProperties=(DeadZone=0.000000,Sensitivity=1.000000,Exponent=1.000000,bInvert=False))
++AxisConfig=(AxisKeyName="OculusTouch_Left_IndexPointing",AxisProperties=(DeadZone=0.000000,Sensitivity=1.000000,Exponent=1.000000,bInvert=False))
++AxisConfig=(AxisKeyName="OculusTouch_Left_ThumbUp",AxisProperties=(DeadZone=0.000000,Sensitivity=1.000000,Exponent=1.000000,bInvert=False))
++AxisConfig=(AxisKeyName="OculusTouch_Left_ThumbRest",AxisProperties=(DeadZone=0.000000,Sensitivity=1.000000,Exponent=1.000000,bInvert=False))
++AxisConfig=(AxisKeyName="OculusTouch_Right_Thumbstick",AxisProperties=(DeadZone=0.000000,Sensitivity=1.000000,Exponent=1.000000,bInvert=False))
++AxisConfig=(AxisKeyName="OculusTouch_Right_FaceButton1",AxisProperties=(DeadZone=0.000000,Sensitivity=1.000000,Exponent=1.000000,bInvert=False))
++AxisConfig=(AxisKeyName="OculusTouch_Right_Trigger",AxisProperties=(DeadZone=0.000000,Sensitivity=1.000000,Exponent=1.000000,bInvert=False))
++AxisConfig=(AxisKeyName="OculusTouch_Right_FaceButton2",AxisProperties=(DeadZone=0.000000,Sensitivity=1.000000,Exponent=1.000000,bInvert=False))
++AxisConfig=(AxisKeyName="OculusTouch_Right_IndexPointing",AxisProperties=(DeadZone=0.000000,Sensitivity=1.000000,Exponent=1.000000,bInvert=False))
++AxisConfig=(AxisKeyName="OculusTouch_Right_ThumbUp",AxisProperties=(DeadZone=0.000000,Sensitivity=1.000000,Exponent=1.000000,bInvert=False))
++AxisConfig=(AxisKeyName="OculusTouch_Right_ThumbRest",AxisProperties=(DeadZone=0.000000,Sensitivity=1.000000,Exponent=1.000000,bInvert=False))
++AxisConfig=(AxisKeyName="OculusHand_Left_ThumbPinchStrength",AxisProperties=(DeadZone=0.000000,Sensitivity=1.000000,Exponent=1.000000,bInvert=False))
++AxisConfig=(AxisKeyName="OculusHand_Left_IndexPinchStrength",AxisProperties=(DeadZone=0.000000,Sensitivity=1.000000,Exponent=1.000000,bInvert=False))
++AxisConfig=(AxisKeyName="OculusHand_Left_MiddlePinchStrength",AxisProperties=(DeadZone=0.000000,Sensitivity=1.000000,Exponent=1.000000,bInvert=False))
++AxisConfig=(AxisKeyName="OculusHand_Left_RingPinchStrength",AxisProperties=(DeadZone=0.000000,Sensitivity=1.000000,Exponent=1.000000,bInvert=False))
++AxisConfig=(AxisKeyName="OculusHand_Left_PinkPinchStrength",AxisProperties=(DeadZone=0.000000,Sensitivity=1.000000,Exponent=1.000000,bInvert=False))
++AxisConfig=(AxisKeyName="OculusHand_Right_ThumbPinchStrength",AxisProperties=(DeadZone=0.000000,Sensitivity=1.000000,Exponent=1.000000,bInvert=False))
++AxisConfig=(AxisKeyName="OculusHand_Right_IndexPinchStrength",AxisProperties=(DeadZone=0.000000,Sensitivity=1.000000,Exponent=1.000000,bInvert=False))
++AxisConfig=(AxisKeyName="OculusHand_Right_MiddlePinchStrength",AxisProperties=(DeadZone=0.000000,Sensitivity=1.000000,Exponent=1.000000,bInvert=False))
++AxisConfig=(AxisKeyName="OculusHand_Right_RingPinchStrength",AxisProperties=(DeadZone=0.000000,Sensitivity=1.000000,Exponent=1.000000,bInvert=False))
++AxisConfig=(AxisKeyName="OculusHand_Right_PinkPinchStrength",AxisProperties=(DeadZone=0.000000,Sensitivity=1.000000,Exponent=1.000000,bInvert=False))
++AxisConfig=(AxisKeyName="Cosmos_Left_Grip_Axis",AxisProperties=(DeadZone=0.000000,Sensitivity=1.000000,Exponent=1.000000,bInvert=False))
++AxisConfig=(AxisKeyName="Cosmos_Left_Trigger_Axis",AxisProperties=(DeadZone=0.000000,Sensitivity=1.000000,Exponent=1.000000,bInvert=False))
++AxisConfig=(AxisKeyName="Cosmos_Left_Thumbstick_X",AxisProperties=(DeadZone=0.000000,Sensitivity=1.000000,Exponent=1.000000,bInvert=False))
++AxisConfig=(AxisKeyName="Cosmos_Left_Thumbstick_Y",AxisProperties=(DeadZone=0.000000,Sensitivity=1.000000,Exponent=1.000000,bInvert=False))
++AxisConfig=(AxisKeyName="Cosmos_Right_Grip_Axis",AxisProperties=(DeadZone=0.000000,Sensitivity=1.000000,Exponent=1.000000,bInvert=False))
++AxisConfig=(AxisKeyName="Cosmos_Right_Trigger_Axis",AxisProperties=(DeadZone=0.000000,Sensitivity=1.000000,Exponent=1.000000,bInvert=False))
++AxisConfig=(AxisKeyName="Cosmos_Right_Thumbstick_X",AxisProperties=(DeadZone=0.000000,Sensitivity=1.000000,Exponent=1.000000,bInvert=False))
++AxisConfig=(AxisKeyName="Cosmos_Right_Thumbstick_Y",AxisProperties=(DeadZone=0.000000,Sensitivity=1.000000,Exponent=1.000000,bInvert=False))
 bAltEnterTogglesFullscreen=True
 bF11TogglesFullscreen=True
 bUseMouseForTouch=False
@@ -35,8 +108,10 @@ FOVScale=0.011110
 DoubleClickTime=0.200000
 +ActionMappings=(ActionName="Jump",bShift=False,bCtrl=False,bAlt=False,bCmd=False,Key=SpaceBar)
 +ActionMappings=(ActionName="Jump",bShift=False,bCtrl=False,bAlt=False,bCmd=False,Key=Gamepad_FaceButton_Bottom)
-+ActionMappings=(ActionName="Fire",bShift=False,bCtrl=False,bAlt=False,bCmd=False,Key=LeftMouseButton)
-+ActionMappings=(ActionName="Fire",bShift=False,bCtrl=False,bAlt=False,bCmd=False,Key=Gamepad_RightTrigger)
++ActionMappings=(ActionName="PrimaryAction",bShift=False,bCtrl=False,bAlt=False,bCmd=False,Key=LeftMouseButton)
++ActionMappings=(ActionName="PrimaryAction",bShift=False,bCtrl=False,bAlt=False,bCmd=False,Key=Gamepad_RightTrigger)
++ActionMappings=(ActionName="Secondary",bShift=False,bCtrl=False,bAlt=False,bCmd=False,Key=RightMouseButton)
++ActionMappings=(ActionName="Secondary",bShift=False,bCtrl=False,bAlt=False,bCmd=False,Key=Gamepad_LeftTrigger)
 +AxisMappings=(AxisName="MoveForward",Scale=1.000000,Key=W)
 +AxisMappings=(AxisName="MoveForward",Scale=-1.000000,Key=S)
 +AxisMappings=(AxisName="MoveForward",Scale=1.000000,Key=Up)
diff --git a/Templates/TP_FirstPersonBP/Config/TemplateDefs.ini b/Templates/TP_FirstPersonBP/Config/TemplateDefs.ini
--- a/Templates/TP_FirstPersonBP/Config/TemplateDefs.ini
+++ b/Templates/TP_FirstPersonBP/Config/TemplateDefs.ini
@@ -38,7 +38,7 @@ FilenameReplacements=(Extensions=("cpp","h","ini","cs"),From="%TEMPLATENAME%",To
 ReplacementsInFiles=(Extensions=("cpp","h","ini","cs"),From="%TEMPLATENAME_UPPERCASE%",To="%PROJECTNAME_UPPERCASE%",bCaseSensitive=true)
 ReplacementsInFiles=(Extensions=("cpp","h","ini","cs"),From="%TEMPLATENAME_LOWERCASE%",To="%PROJECTNAME_LOWERCASE%",bCaseSensitive=true)
 ReplacementsInFiles=(Extensions=("cpp","h","ini","cs"),From="%TEMPLATENAME%",To="%PROJECTNAME%",bCaseSensitive=false)
-SharedContentPacks=(MountName="FirstPerson",DetailLevels=("Standard"))
-SharedContentPacks=(MountName="Geometry",DetailLevels=("Standard","High"))
+SharedContentPacks=(MountName="FirstPersonArms",DetailLevels=("Standard"))
 SharedContentPacks=(MountName="FPWeapon",DetailLevels=("Standard"))
+SharedContentPacks=(MountName="LevelPrototyping",DetailLevels=("Standard","High"))
 EditDetailLevelPreference="High"
\ No newline at end of file
diff --git a/Templates/TP_FirstPersonBP/manifest.json b/Templates/TP_FirstPersonBP/manifest.json
--- a/Templates/TP_FirstPersonBP/manifest.json
+++ b/Templates/TP_FirstPersonBP/manifest.json
@@ -62,20 +62,19 @@
 	],
 	"AdditionalFeaturePacks": [
 		{
-			"MountName": "Geometry",
+			"MountName": "FPWeapon",
 			"DetailLevels": [
-				"Standard",
-				"High"
+				"Standard"
 			]
 		},
 		{
-			"MountName": "FPWeapon",
+			"MountName": "LevelPrototyping",
 			"DetailLevels": [
-				"Standard"
+				"High"
 			]
 		},
 		{
-			"MountName": "FirstPerson",
+			"MountName": "FirstPersonArms",
 			"DetailLevels":
 			[
 				"Standard"
diff --git a/Templates/TemplateResources/Standard/FirstPerson/FeaturePack/manifest.json b/Templates/TemplateResources/Standard/FirstPerson/FeaturePack/manifest.json
deleted file mode 100644
diff --git a/Templates/TemplateResources/Standard/FirstPersonArms/FeaturePack/manifest.json b/Templates/TemplateResources/Standard/FirstPersonArms/FeaturePack/manifest.json
new file mode 100644
--- /dev/null
+++ b/Templates/TemplateResources/Standard/FirstPersonArms/FeaturePack/manifest.json
@@ -0,0 +1,48 @@
+{
+	"Version": 1,
+	"Name":
+	[
+		{
+			"Language": "en",
+			"Text": "FirstPerson Shared Resource Standard"
+		}
+	],
+	"Description":
+	[
+		{
+			"Language": "en",
+			"Text": "Shared Resource Pack for FirstPerson"
+		}
+	],
+	"AssetTypes":
+	[
+		{
+			"Language": "en",
+			"Text": ""
+		}
+	],
+	"SearchTags":
+	[
+		{
+			"Language": "en",
+			"Text": ""
+		}
+	],	
+	"ClassTypes": " ",
+	"Category": "SharedPack",
+	"Thumbnail": "FirstPerson.png",
+	"Screenshots":
+	[
+		"FirstPerson_Preview.png"
+	],
+	"Ident": "FirstPersonArmsStandard",
+	"AdditionalFiles":
+	{
+		"DestinationFilesFolder": "FirstPersonArms",
+		"AdditionalFilesList":
+		[
+			"Templates/TemplateResources/Standard/FirstPersonArms/Content/*.*"
+		]
+	}
+}
+
