First Person BP Template Level refresh with new level prototyping content.
Moved input logic back to FirstPersonCharacter.

diff --git a/Templates/TP_FirstPersonBP/Config/DefaultInput.ini b/Templates/TP_FirstPersonBP/Config/DefaultInput.ini
--- a/Templates/TP_FirstPersonBP/Config/DefaultInput.ini
+++ b/Templates/TP_FirstPersonBP/Config/DefaultInput.ini
@@ -110,22 +110,20 @@ DoubleClickTime=0.200000
 +ActionMappings=(ActionName="Jump",bShift=False,bCtrl=False,bAlt=False,bCmd=False,Key=Gamepad_FaceButton_Bottom)
 +ActionMappings=(ActionName="PrimaryAction",bShift=False,bCtrl=False,bAlt=False,bCmd=False,Key=LeftMouseButton)
 +ActionMappings=(ActionName="PrimaryAction",bShift=False,bCtrl=False,bAlt=False,bCmd=False,Key=Gamepad_RightTrigger)
-+ActionMappings=(ActionName="Secondary",bShift=False,bCtrl=False,bAlt=False,bCmd=False,Key=RightMouseButton)
-+ActionMappings=(ActionName="Secondary",bShift=False,bCtrl=False,bAlt=False,bCmd=False,Key=Gamepad_LeftTrigger)
-+AxisMappings=(AxisName="MoveForward",Scale=1.000000,Key=W)
-+AxisMappings=(AxisName="MoveForward",Scale=-1.000000,Key=S)
-+AxisMappings=(AxisName="MoveForward",Scale=1.000000,Key=Up)
-+AxisMappings=(AxisName="MoveForward",Scale=-1.000000,Key=Down)
-+AxisMappings=(AxisName="MoveForward",Scale=1.000000,Key=Gamepad_LeftY)
-+AxisMappings=(AxisName="MoveRight",Scale=-1.000000,Key=A)
-+AxisMappings=(AxisName="MoveRight",Scale=1.000000,Key=D)
-+AxisMappings=(AxisName="MoveRight",Scale=1.000000,Key=Gamepad_LeftX)
-+AxisMappings=(AxisName="TurnRate",Scale=1.000000,Key=Gamepad_RightX)
-+AxisMappings=(AxisName="TurnRate",Scale=-1.000000,Key=Left)
-+AxisMappings=(AxisName="TurnRate",Scale=1.000000,Key=Right)
-+AxisMappings=(AxisName="Turn",Scale=1.000000,Key=MouseX)
-+AxisMappings=(AxisName="LookUpRate",Scale=1.000000,Key=Gamepad_RightY)
-+AxisMappings=(AxisName="LookUp",Scale=-1.000000,Key=MouseY)
++AxisMappings=(AxisName="Move Forward / Backward",Scale=1.000000,Key=W)
++AxisMappings=(AxisName="Move Forward / Backward",Scale=-1.000000,Key=S)
++AxisMappings=(AxisName="Move Forward / Backward",Scale=1.000000,Key=Up)
++AxisMappings=(AxisName="Move Forward / Backward",Scale=-1.000000,Key=Down)
++AxisMappings=(AxisName="Move Forward / Backward",Scale=1.000000,Key=Gamepad_LeftY)
++AxisMappings=(AxisName="Move Right / Left",Scale=-1.000000,Key=A)
++AxisMappings=(AxisName="Move Right / Left",Scale=1.000000,Key=D)
++AxisMappings=(AxisName="Move Right / Left",Scale=1.000000,Key=Gamepad_LeftX)
++AxisMappings=(AxisName="Turn Right / Left Gamepad",Scale=1.000000,Key=Gamepad_RightX)
++AxisMappings=(AxisName="Turn Right / Left Gamepad",Scale=-1.000000,Key=Left)
++AxisMappings=(AxisName="Turn Right / Left Gamepad",Scale=1.000000,Key=Right)
++AxisMappings=(AxisName="Turn Right / Left Mouse",Scale=1.000000,Key=MouseX)
++AxisMappings=(AxisName="Look Up / Down Gamepad",Scale=1.000000,Key=Gamepad_RightY)
++AxisMappings=(AxisName="Look Up / Down Mouse",Scale=-1.000000,Key=MouseY)
 DefaultPlayerInputClass=/Script/Engine.PlayerInput
 DefaultInputComponentClass=/Script/Engine.InputComponent
 DefaultTouchInterface=/Engine/MobileResources/HUD/DefaultVirtualJoysticks.DefaultVirtualJoysticks
