Implemented a custom Touch Interface to fix an issue with conflicting Jump and Shoot mobile controls on First Person template.

diff --git a/Templates/TP_FirstPerson/Config/DefaultInput.ini b/Templates/TP_FirstPerson/Config/DefaultInput.ini
--- a/Templates/TP_FirstPerson/Config/DefaultInput.ini
+++ b/Templates/TP_FirstPerson/Config/DefaultInput.ini
@@ -103,7 +103,7 @@ FOVScale=0.011110
 DoubleClickTime=0.200000
 DefaultPlayerInputClass=/Script/EnhancedInput.EnhancedPlayerInput
 DefaultInputComponentClass=/Script/EnhancedInput.EnhancedInputComponent
-DefaultTouchInterface=/Engine/MobileResources/HUD/DefaultVirtualJoysticks.DefaultVirtualJoysticks
+DefaultTouchInterface=/Game/FirstPerson/Input/MobileControls.MobileControls
 -ConsoleKeys=Tilde
 +ConsoleKeys=Tilde
 
diff --git a/Templates/TP_FirstPersonBP/Config/DefaultInput.ini b/Templates/TP_FirstPersonBP/Config/DefaultInput.ini
--- a/Templates/TP_FirstPersonBP/Config/DefaultInput.ini
+++ b/Templates/TP_FirstPersonBP/Config/DefaultInput.ini
@@ -111,7 +111,7 @@ FOVScale=0.011110
 DoubleClickTime=0.200000
 DefaultPlayerInputClass=/Script/EnhancedInput.EnhancedPlayerInput
 DefaultInputComponentClass=/Script/EnhancedInput.EnhancedInputComponent
-DefaultTouchInterface=/Engine/MobileResources/HUD/DefaultVirtualJoysticks.DefaultVirtualJoysticks
+DefaultTouchInterface=/Game/FirstPerson/Input/MobileControls.MobileControls
 -ConsoleKeys=Tilde
 +ConsoleKeys=Tilde
 
