Removed old UE4 Mannuequin from Characters folder in templates
Updated Manny and Quinn folder to include reduced texture sizes, additional animations etc.
Removed Quinn animation
Updated references in templates to match new structure in characters folder
Added Shooter Variant
Added Horror Variant
Added New Weapons Pack
Updated manifest and enginedefs for first person
Updated first person bp template to support Variants

diff --git a/Templates/TP_FirstPersonBP/Config/DefaultEngine.ini b/Templates/TP_FirstPersonBP/Config/DefaultEngine.ini
--- a/Templates/TP_FirstPersonBP/Config/DefaultEngine.ini
+++ b/Templates/TP_FirstPersonBP/Config/DefaultEngine.ini
@@ -4,14 +4,14 @@
 +EditProfiles=(Name="Trigger",CustomResponses=((Channel=Projectile, Response=ECR_Ignore)))
 
 [/Script/EngineSettings.GameMapsSettings]
-EditorStartupMap=/Game/FirstPerson/Maps/FirstPersonMap.FirstPersonMap
+EditorStartupMap=/Game/FirstPerson/Lvl_FirstPerson.Lvl_FirstPerson
 LocalMapOptions=
 TransitionMap=
 bUseSplitscreen=True
 TwoPlayerSplitscreenLayout=Horizontal
 ThreePlayerSplitscreenLayout=FavorTop
 GameInstanceClass=/Script/Engine.GameInstance
-GameDefaultMap=/Game/FirstPerson/Maps/FirstPersonMap.FirstPersonMap
+GameDefaultMap=/Game/FirstPerson/Lvl_FirstPerson.Lvl_FirstPerson
 ServerDefaultMap=/Engine/Maps/Entry
 GlobalDefaultGameMode=/Game/FirstPerson/Blueprints/BP_FirstPersonGameMode.BP_FirstPersonGameMode_C
 GlobalDefaultServerGameMode=None
@@ -20,8 +20,25 @@ GlobalDefaultServerGameMode=None
 r.ReflectionMethod=1
 r.GenerateMeshDistanceFields=True
 r.DynamicGlobalIlluminationMethod=1
+r.Mobile.EnableNoPrecomputedLightingCSMShader=1
+r.DefaultFeature.AutoExposure.ExtendDefaultLuminanceRange=True
 r.DefaultFeature.AutoExposure.ExtendDefaultLuminanceRange=true
+
+r.Shadow.Virtual.Enable=1
+
+r.DefaultFeature.LocalExposure.HighlightContrastScale=0.8
+
+r.DefaultFeature.LocalExposure.ShadowContrastScale=0.8
+r.DefaultFeature.LightUnits=2
 r.AllowStaticLighting=False
+RedChromaticityCoordinate=(X=0.640000,Y=0.330000)
+GreenChromaticityCoordinate=(X=0.300000,Y=0.600000)
+BlueChromaticityCoordinate=(X=0.150000,Y=0.060000)
+WhiteChromaticityCoordinate=(X=0.312700,Y=0.329000)
+WorkingColorSpaceChoice=sRGB
+r.LegacyLuminanceFactors=False
+r.Lumen.ScreenTracingSource=0
+r.FirstPerson.SelfShadow 1
 
 [/Script/WindowsTargetPlatform.WindowsTargetSettings]
 DefaultGraphicsRHI=DefaultGraphicsRHI_DX12
@@ -30,6 +47,31 @@ DefaultGraphicsRHI=DefaultGraphicsRHI_DX12
 -D3D11TargetedShaderFormats=PCD3D_SM5
 +D3D11TargetedShaderFormats=PCD3D_SM5
 
-[/Script/LinuxTargetPlatform.LinuxTargetSettings]
--TargetedRHIs=SF_VULKAN_SM5
-+TargetedRHIs=SF_VULKAN_SM6
\ No newline at end of file
+[/Script/HardwareTargeting.HardwareTargetingSettings]
+TargetedHardwareClass=Desktop
+AppliedTargetedHardwareClass=Desktop
+DefaultGraphicsPerformance=Maximum
+AppliedDefaultGraphicsPerformance=Maximum
+
+[/Script/Engine.Engine]
++ActiveGameNameRedirects=(OldGameName="TP_FirstPersonBP",NewGameName="/Script/MyProject")
++ActiveGameNameRedirects=(OldGameName="/Script/TP_FirstPersonBP",NewGameName="/Script/MyProject")
+NearClipPlane=0.500000
+
+[/Script/AndroidFileServerEditor.AndroidFileServerRuntimeSettings]
+bEnablePlugin=True
+bAllowNetworkConnection=True
+SecurityToken=85DFF0BB498AE9EFE914DE8DB82D122D
+bIncludeInShipping=False
+bAllowExternalStartInShipping=False
+bCompileAFSProject=False
+bUseCompression=False
+bLogFiles=False
+bReportStats=False
+ConnectionType=USBOnly
+bUseManualIPAddress=False
+ManualIPAddress=
+
+[/Script/AIModule.AISystem]
+bForgetStaleActors=True
+
diff --git a/Templates/TP_FirstPersonBP/Config/DefaultGame.ini b/Templates/TP_FirstPersonBP/Config/DefaultGame.ini
--- a/Templates/TP_FirstPersonBP/Config/DefaultGame.ini
+++ b/Templates/TP_FirstPersonBP/Config/DefaultGame.ini
@@ -1,3 +1,3 @@
 [/Script/EngineSettings.GeneralProjectSettings]
-ProjectID=5791E5B14A4873D64314CB8851F7814D
+ProjectID=048A64D545DC87F5AF2B3B9F89DBFECB
 ProjectName=First Person BP Game Template
diff --git a/Templates/TP_FirstPersonBP/Config/TemplateDefs.ini b/Templates/TP_FirstPersonBP/Config/TemplateDefs.ini
--- a/Templates/TP_FirstPersonBP/Config/TemplateDefs.ini
+++ b/Templates/TP_FirstPersonBP/Config/TemplateDefs.ini
@@ -41,7 +41,9 @@ FilenameReplacements=(Extensions=("cpp","h","ini","cs"),From="%TEMPLATENAME%",To
 ReplacementsInFiles=(Extensions=("cpp","h","ini","cs"),From="%TEMPLATENAME_UPPERCASE%",To="%PROJECTNAME_UPPERCASE%",bCaseSensitive=true)
 ReplacementsInFiles=(Extensions=("cpp","h","ini","cs"),From="%TEMPLATENAME_LOWERCASE%",To="%PROJECTNAME_LOWERCASE%",bCaseSensitive=true)
 ReplacementsInFiles=(Extensions=("cpp","h","ini","cs"),From="%TEMPLATENAME%",To="%PROJECTNAME%",bCaseSensitive=false)
-SharedContentPacks=(MountName="FirstPersonArms",DetailLevels=("Standard"))
-SharedContentPacks=(MountName="FPWeapon",DetailLevels=("Standard"))
+SharedContentPacks=(MountName="Characters",DetailLevels=("High"))
 SharedContentPacks=(MountName="LevelPrototyping",DetailLevels=("High"))
-EditDetailLevelPreference="High"
\ No newline at end of file
+EditDetailLevelPreference="High"
+
+Variants=(Name="ArenaShooter",LocalizedDisplayNames=((Language="en",Text="Arena Shooter"),(Language="ko",Text="퍼즐")),LocalizedDescriptions=((Language="en",Text="An arena shooter with weapon pickups and an enemy AI opponent."),(Language="ko",Text="퍼즐 템플릿은 단순한 블록 모눈을 특징으로 합니다.")),SharedContentPacks=((DetailLevels=(Standard),MountName="Weapons"),(DetailLevels=(Standard),MountName="Variant_Shooter")) )
+Variants=(Name="SurvivalHorror",LocalizedDisplayNames=((Language="en",Text="Survival Horror"),(Language="ko",Text="퍼즐")),LocalizedDescriptions=((Language="en",Text="A spooky map."),(Language="ko",Text="퍼즐 템플릿은 단순한 블록 모눈을 특징으로 합니다.")),SharedContentPacks=((DetailLevels=(Standard),MountName="Variant_Horror")) )
\ No newline at end of file
diff --git a/Templates/TP_FirstPersonBP/TP_FirstPersonBP.uproject b/Templates/TP_FirstPersonBP/TP_FirstPersonBP.uproject
--- a/Templates/TP_FirstPersonBP/TP_FirstPersonBP.uproject
+++ b/Templates/TP_FirstPersonBP/TP_FirstPersonBP.uproject
@@ -10,6 +10,10 @@
 			"TargetAllowList": [
 				"Editor"
 			]
+		},
+		{
+			"Name": "GameplayStateTree",
+			"Enabled": true
 		}
 	]	
 }
\ No newline at end of file
diff --git a/Templates/TemplateResources/Standard/Variant_Horror/FeaturePack/manifest.json b/Templates/TemplateResources/Standard/Variant_Horror/FeaturePack/manifest.json
new file mode 100644
--- /dev/null
+++ b/Templates/TemplateResources/Standard/Variant_Horror/FeaturePack/manifest.json
@@ -0,0 +1,41 @@
+{
+	"Version": 1,
+	"Name": [
+		{
+			"Language": "en",
+			"Text": "FirstPerson Shared Resource Standard"
+		}
+	],
+	"Description": [
+		{
+			"Language": "en",
+			"Text": "Shared Resource Pack for FirstPerson"
+		}
+	],
+	"AssetTypes": [
+		{
+			"Language": "en",
+			"Text": ""
+		}
+	],
+	"SearchTags": [
+		{
+			"Language": "en",
+			"Text": ""
+		}
+	],
+	"ClassTypes": " ",
+	"Category": "SharedPack",
+	"Thumbnail": "Variant_Horror.png",
+	"Screenshots": [
+		"Variant_Horror_Preview.png"
+	],
+	"Ident": "Variant_HorrorStandard",
+	"AdditionalFiles": {
+		"DestinationFilesFolder": "Variant_Horror",
+		"AdditionalFilesList": [
+			"Templates/TemplateResources/Standard/Variant_Horror/Content/*.*"
+		]
+	}
+}
+
diff --git a/Templates/TemplateResources/Standard/Variant_Shooter/FeaturePack/manifest.json b/Templates/TemplateResources/Standard/Variant_Shooter/FeaturePack/manifest.json
new file mode 100644
--- /dev/null
+++ b/Templates/TemplateResources/Standard/Variant_Shooter/FeaturePack/manifest.json
@@ -0,0 +1,41 @@
+{
+	"Version": 1,
+	"Name": [
+		{
+			"Language": "en",
+			"Text": "FirstPerson Shared Resource Standard"
+		}
+	],
+	"Description": [
+		{
+			"Language": "en",
+			"Text": "Shared Resource Pack for FirstPerson"
+		}
+	],
+	"AssetTypes": [
+		{
+			"Language": "en",
+			"Text": ""
+		}
+	],
+	"SearchTags": [
+		{
+			"Language": "en",
+			"Text": ""
+		}
+	],
+	"ClassTypes": " ",
+	"Category": "SharedPack",
+	"Thumbnail": "Variant_Shooter.png",
+	"Screenshots": [
+		"Variant_Shooter_Preview.png"
+	],
+	"Ident": "Variant_ShooterStandard",
+	"AdditionalFiles": {
+		"DestinationFilesFolder": "Variant_Shooter",
+		"AdditionalFilesList": [
+			"Templates/TemplateResources/Standard/Variant_Shooter/Content/*.*"
+		]
+	}
+}
+
diff --git a/Templates/TemplateResources/Standard/Weapons/FeaturePack/manifest.json b/Templates/TemplateResources/Standard/Weapons/FeaturePack/manifest.json
new file mode 100644
--- /dev/null
+++ b/Templates/TemplateResources/Standard/Weapons/FeaturePack/manifest.json
@@ -0,0 +1,48 @@
+{
+	"Version": 1,
+	"Name":
+	[
+		{
+			"Language": "en",
+			"Text": "Weapons"
+		}
+	],
+	"Description":
+	[
+		{
+			"Language": "en",
+			"Text": "Shared Resource Pack for Weapons"
+		}
+	],
+	"AssetTypes":
+	[
+		{
+			"Language": "en",
+			"Text": ""
+		}
+	],
+	"SearchTags":
+	[
+		{
+			"Language": "en",
+			"Text": ""
+		}
+	],	
+	"ClassTypes": " ",
+	"Category": "SharedPack",
+	"Thumbnail": "FPWeapon.png",
+	"Screenshots":
+	[
+		"FPWeapon_Preview.png"
+	],
+	"Ident": "WeaponsStandard",
+	"AdditionalFiles":
+	{
+		"DestinationFilesFolder": "Weapons",
+		"AdditionalFilesList":
+		[
+			"Templates/TemplateResources/Standard/Weapons/Content/*.*"
+		]
+	}
+}
+
