Templates - ThirdPersonBP
Fixed an issue with the SideScroller NPC StateTree behavior.
Added convenience SharedContentPacks comments to FirstPersonBP to aid with variant edits and bug fixes.
#rnx

diff --git a/Templates/TP_FirstPersonBP/Config/TemplateDefs.ini b/Templates/TP_FirstPersonBP/Config/TemplateDefs.ini
--- a/Templates/TP_FirstPersonBP/Config/TemplateDefs.ini
+++ b/Templates/TP_FirstPersonBP/Config/TemplateDefs.ini
@@ -45,5 +45,12 @@ SharedContentPacks=(MountName="Characters",DetailLevels=("High"))
 SharedContentPacks=(MountName="LevelPrototyping",DetailLevels=("High"))
 EditDetailLevelPreference="High"
 
+; Variant packs
+; NOTE: make sure these are commented out for release!
+;SharedContentPacks=(MountName="Variant_Shooter",DetailLevels=("Standard"))
+;SharedContentPacks=(MountName="Variant_Horror",DetailLevels=("Standard"))
+
+; Variant definitions
 Variants=(Name="ArenaShooter",LocalizedDisplayNames=((Language="en",Text="Arena Shooter"),(Language="ko",Text="퍼즐")),LocalizedDescriptions=((Language="en",Text="An arena shooter with weapon pickups and an enemy AI opponent."),(Language="ko",Text="퍼즐 템플릿은 단순한 블록 모눈을 특징으로 합니다.")),SharedContentPacks=((DetailLevels=(Standard),MountName="Weapons"),(DetailLevels=(Standard),MountName="Variant_Shooter")) )
+
 Variants=(Name="SurvivalHorror",LocalizedDisplayNames=((Language="en",Text="Survival Horror"),(Language="ko",Text="퍼즐")),LocalizedDescriptions=((Language="en",Text="A spooky map."),(Language="ko",Text="퍼즐 템플릿은 단순한 블록 모눈을 특징으로 합니다.")),SharedContentPacks=((DetailLevels=(Standard),MountName="Variant_Horror")) )
\ No newline at end of file
