Fixed r. command in default engine for shooter
Fixed shooter grid alignment
Fixed grenade launcher texture size

diff --git a/Templates/TP_FirstPersonBP/Config/DefaultEngine.ini b/Templates/TP_FirstPersonBP/Config/DefaultEngine.ini
--- a/Templates/TP_FirstPersonBP/Config/DefaultEngine.ini
+++ b/Templates/TP_FirstPersonBP/Config/DefaultEngine.ini
@@ -38,7 +38,7 @@ WhiteChromaticityCoordinate=(X=0.312700,Y=0.329000)
 WorkingColorSpaceChoice=sRGB
 r.LegacyLuminanceFactors=False
 r.Lumen.ScreenTracingSource=0
-r.FirstPerson.SelfShadow 1
+r.FirstPerson.SelfShadow=1
 
 [/Script/WindowsTargetPlatform.WindowsTargetSettings]
 DefaultGraphicsRHI=DefaultGraphicsRHI_DX12
