Updated template defs, removed loc from descriptions and double checked variants enabled.
Also includes new icons in media
Updated shooter kill z to be closer to 0.

diff --git a/Templates/TP_FirstPersonBP/Config/TemplateDefs.ini b/Templates/TP_FirstPersonBP/Config/TemplateDefs.ini
--- a/Templates/TP_FirstPersonBP/Config/TemplateDefs.ini
+++ b/Templates/TP_FirstPersonBP/Config/TemplateDefs.ini
@@ -5,15 +5,6 @@ Categories=Games
 LocalizedDisplayNames=(Language="en",Text="First Person")
 LocalizedDescriptions=(Language="en",Text="The first person template features a player character (represented by a pair of arms) which is viewed from first person perspective. The character can be moved around the level using keyboard, controller or virtual joystick on a touch device. Additionally the player can look around using mouse, controller or virtual joystick on a touch device. The character is also equipped with a gun that, using mouse,  controller or virtual joystick on a touch device will fire a simple sphere projectile that will affect some physics objects in the level, whilst bounce off the arena walls.")
 
-LocalizedDisplayNames=(Language="ko",Text="일인칭")
-LocalizedDescriptions=(Language="ko",Text="일인칭 템플릿은 일인칭 시점에서 보이는 (양팔로 표현되는) 플레이어 캐릭터를 특징으로 합니다. 캐릭터 이동은 키보드, 컨트롤러, 터치 디바이스의 가상 조이스틱으로 가능합니다. 추가적으로 플레이어는 마우스, 컨트롤러, 터치 디바이스의 가상 조이스틱을 사용하여 주변을 둘러볼 수도 있습니다. 캐릭터에는 총이 장착되어 있어 마우스, 컨트롤러, 터치 디바이스의 가상 조이스틱을 사용하면 간단한 발사체가 발사되어, 아레나 벽에 튕기면서 레벨에 있는 몇몇 물리 오브젝트에 영향을 줍니다.")
-
-LocalizedDisplayNames=(Language="ja",Text="ファーストパーソン")
-LocalizedDescriptions=(Language="ja",Text="ファーストパーソンテンプレートには、一人称視点で表示されるプレイヤーキャラクタ（腕として表示されています）がいます。キーボード、コントローラー、タッチデバイスでのバーチャルジョイスティックでこのキャラクタを移動させることができます。追加して、プレイヤーはマウス、コントローラー、タッチデバイスでのバーチャルジョイスティックで周りを見回すことができます。さらにキャラクタは銃を装備して、マウス、コントローラー、タッチデバイスでのバーチャルジョイスティックを使ってシンプルな球体の弾を発射できます。この弾は物理オブジェクトに影響し、アリーナの壁で反射します。")
-
-LocalizedDisplayNames=(Language="zh-Hans",Text="第一人称游戏")
-LocalizedDescriptions=(Language="zh-Hans",Text="第一人称模板包含一个第一人称玩家角色（通过一双手臂呈现）。你可以通过键盘、手柄或触摸式设备上的虚拟摇杆控制角色在关卡中四处移动。你还可以使用鼠标、手柄或触摸式设备上的虚拟摇杆转动角色视角，观察四周。此外，你还可以让角色拾取枪支，然后用鼠标、手柄或触摸式设备上的虚拟摇杆来发射球形子弹。子弹会撞飞关卡中的物理对象，还能从竞技场墙壁上弹开。")
-
 ; This project uses a thumbnail as the project icon
 bThumbnailAsIcon=true
 
@@ -50,8 +41,9 @@ EditDetailLevelPreference="High"
 ; NOTE: make sure these are commented out for release!
 ;SharedContentPacks=(MountName="Variant_Shooter",DetailLevels=("Standard"))
 ;SharedContentPacks=(MountName="Variant_Horror",DetailLevels=("Standard"))
+;SharedContentPacks=(MountName="Weapons",DetailLevels=("Standard"))
 
 ; Variant definitions
-Variants=(Name="ArenaShooter",LocalizedDisplayNames=((Language="en",Text="Arena Shooter"),(Language="ko",Text="퍼즐")),LocalizedDescriptions=((Language="en",Text="An arena shooter with weapon pickups and an enemy AI opponent."),(Language="ko",Text="퍼즐 템플릿은 단순한 블록 모눈을 특징으로 합니다.")),SharedContentPacks=((DetailLevels=(Standard),MountName="Weapons"),(DetailLevels=(Standard),MountName="Variant_Shooter")) )
+Variants=(Name="ArenaShooter",LocalizedDisplayNames=((Language="en",Text="Arena Shooter")),LocalizedDescriptions=((Language="en",Text="An arena shooter with weapon pickups and an enemy AI opponent.")),SharedContentPacks=((DetailLevels=(Standard),MountName="Weapons"),(DetailLevels=(Standard),MountName="Variant_Shooter")) )
 
-Variants=(Name="SurvivalHorror",LocalizedDisplayNames=((Language="en",Text="Survival Horror"),(Language="ko",Text="퍼즐")),LocalizedDescriptions=((Language="en",Text="A spooky map."),(Language="ko",Text="퍼즐 템플릿은 단순한 블록 모눈을 특징으로 합니다.")),SharedContentPacks=((DetailLevels=(Standard),MountName="Variant_Horror")) )
\ No newline at end of file
+Variants=(Name="SurvivalHorror",LocalizedDisplayNames=((Language="en",Text="Survival Horror")),LocalizedDescriptions=((Language="en",Text="A simple first person character with sprint mechanic and torch set in a dark and suspensful map.")),SharedContentPacks=((DetailLevels=(Standard),MountName="Variant_Horror")) )
diff --git a/Templates/TP_ThirdPersonBP/Config/TemplateDefs.ini b/Templates/TP_ThirdPersonBP/Config/TemplateDefs.ini
--- a/Templates/TP_ThirdPersonBP/Config/TemplateDefs.ini
+++ b/Templates/TP_ThirdPersonBP/Config/TemplateDefs.ini
@@ -3,16 +3,7 @@
 Categories=Games
 
 LocalizedDisplayNames=(Language="en",Text="Third Person")
-LocalizedDescriptions=(Language="en",Text="The third person pack features a playable character where the camera is positioned behind and slightly above the character. As the character moves using mouse, keyboard, controller or virtual joystick on a touch device, the camera follows the character, observing from an over-the-shoulder point of view. This perspective emphasizes the main character and is popular in action and adventure games. The template map contains several objects that the player can jump upon, and push around.")
-
-LocalizedDisplayNames=(Language="ko",Text="삼인칭")
-LocalizedDescriptions=(Language="ko",Text="삼인칭 템플릿에는 플레이 가능한 캐릭터가 포함되어 있으며 카메라가 캐릭터 머리 위 약간 뒤쪽에 위치합니다. 마우스, 키보드, 터치 디바이스의 가상 조이스틱으로 캐릭터가 이동함에 따라 카메라가 따라 움직이며 어깨 위 시점으로 봅니다. 이 시점은 주인공 캐릭터를 강조해 액션이나 어드벤처 게임에 주로 활용됩니다. 이 템플릿 맵에는 플레이어가 점프하여 착지하고 밀 수 있는 여러가지 오브젝트가 포함되어 있습니다.")
-
-LocalizedDisplayNames=(Language="ja",Text="サードパーソン")
-LocalizedDescriptions=(Language="ja",Text="サードパーソン パックは、後ろ側と少し上にカメラが配置されているプレイ可能なキャラクターを備えています。タッチデバイスでマウス、キーボード、コントローラー、またはバーチャル ジョイスティックの使用によりキャラクターが移動すると、カメラはキャラクターを追跡し、肩越しの視点から観察します。この視点では主人公が強調されるため、アクションゲームやアドベンチャー ゲームで人気があります。テンプレート マップには、プレイヤーがジャンプしたり押すことができるオブジェクトがいくつか含まれています。")
-
-LocalizedDisplayNames=(Language="zh-Hans",Text="第三人称游戏")
-LocalizedDescriptions=(Language="zh-Hans",Text="第三人称游戏内容包包含一个可操控的角色以及一台位于角色身后斜上方的摄像机。你可以通过鼠标、键盘、手柄或触摸式设备上控制角色移动。摄像机会跟随角色，以越肩视角拍摄场景。这种视角突出主角，常用于动作及冒险类游戏。该模板的地图中还包含一些对象，可以让角色跳上跳下或四处推动。")
+LocalizedDescriptions=(Language="en",Text="The third person pack features a playable character where the camera is positioned behind and slightly above the character. As the character moves using mouse, keyboard, controller or virtual joystick on a touch device, the camera follows the character, observing from an over-the-shoulder point of view. Using variants, this template can be expanded to work as a 2.5D sidescroller, a traditional platformer or a hand to hand combat game.")
 
 ; This project uses a thumbnail as the project icon
 bThumbnailAsIcon=true
diff --git a/Templates/TP_TopDownBP/Config/TemplateDefs.ini b/Templates/TP_TopDownBP/Config/TemplateDefs.ini
--- a/Templates/TP_TopDownBP/Config/TemplateDefs.ini
+++ b/Templates/TP_TopDownBP/Config/TemplateDefs.ini
@@ -6,15 +6,6 @@ Categories=Games
 LocalizedDisplayNames=(Language="en",Text="Top Down")
 LocalizedDescriptions=(Language="en",Text="The top down template features a playable character where the camera is positioned behind and above the avatar at greater distance than in third person perspective. The player is controlled by using the mouse, gamepad, or touchscreen and moves the character by clicking on a destination, or following the cursor if input is held. This perspective is often used in action role playing games.")
 
-LocalizedDisplayNames=(Language="ko",Text="탑다운")
-LocalizedDescriptions=(Language="ko",Text="탑다운 템플릿에는 플레이 가능한 캐릭터가 포함되어 있으며 카메라가 아바타 머리 위 약간 뒤쪽에서 삼인칭 시점보다 더 원거리에 위치합니다. 플레이어는 마우스, 게임패드나 터치스크린으로 제어되고, 목적지를 클릭하거나 입력을 유지하여 커서에 따라 캐릭터를 움직일 수 있습니다. 이 시점은 주로 액션 RPG에 활용됩니다.")
-
-LocalizedDisplayNames=(Language="ja",Text="トップダウン")
-LocalizedDescriptions=(Language="ja",Text="トップダウン テンプレートは、カメラがアバターの後ろと上に配置されている (サードパーソン視点よりも距離が遠い位置) プレイ可能なキャラクターを特徴としています。プレイヤーは、マウス、ゲームパッド、またはタッチ スクリーンの使用により制御され、目的地をクリックするか、入力が保持されている場合はカーソルをたどることでキャラクターを移動させます。この視点は、アクション ロールプレイ ゲームでよく使用されます。")
-
-LocalizedDisplayNames=(Language="zh-Hans",Text="俯视角游戏")
-LocalizedDescriptions=(Language="zh-Hans",Text="俯视角游戏模板包含一个可操控的角色以及一台位于角色后上方的摄像机。摄像机的位置比第三人称视角模板中的要高很多。玩家可以通过点击鼠标、触控板或触摸屏来选择目的地并移动角色，或持续按下鼠标来移动角色。此视角常用于动作角色扮演类游戏。")
-
 ; This project uses a thumbnail as the project icon
 bThumbnailAsIcon=true
 
@@ -58,4 +49,5 @@ EditDetailLevelPreference="High"
 
 ; Variant definitions
 Variants=(Name="Strategy",LocalizedDisplayNames=((Language="en",Text="Strategy")),LocalizedDescriptions=((Language="en",Text="A top down strategy game with point and click interactions and selectable units.")),SharedContentPacks=((DetailLevels=(Standard),MountName="Variant_Strategy")))
+
 Variants=(Name="TwinStick",LocalizedDisplayNames=((Language="en",Text="Twin Stick")),LocalizedDescriptions=((Language="en",Text="A twin stick shooter game with infinitely spawning enemies.")),SharedContentPacks=((DetailLevels=(Standard),MountName="Variant_TwinStick"),(MountName="Input",DetailLevels=("High")))
\ No newline at end of file
diff --git a/Templates/TP_VehicleAdvBP/Config/TemplateDefs.ini b/Templates/TP_VehicleAdvBP/Config/TemplateDefs.ini
--- a/Templates/TP_VehicleAdvBP/Config/TemplateDefs.ini
+++ b/Templates/TP_VehicleAdvBP/Config/TemplateDefs.ini
@@ -5,15 +5,6 @@ Categories=Games
 LocalizedDisplayNames=(Language="en",Text="Vehicle")
 LocalizedDescriptions=(Language="en",Text="The vehicle pack features both a regular vehicle and a more complex vehicle. Vehicle movement can be controlled with either keyboard, controller or virtual joystick on a touch device. The advanced vehicle itself features working double wishbone suspension. There is also an engine sound in this template. There is also a HUD that shows the current gear and speed.")
 
-LocalizedDisplayNames=(Language="ko",Text="비히클")
-LocalizedDescriptions=(Language="ko",Text="비히클 팩은 일반 비히클과 조금 더 복잡한 비히클을 모두 제공합니다. 비히클의 이동은 키보드, 컨트롤러 또는 터치 디바이스의 버추얼 조이스틱으로 제어 가능합니다. 고급 비히클은 더블 위시본 서스펜션을 갖추고 있습니다. 이 템플릿에는 엔진 사운드와, 현재 기어와 속도를 보여주는 HUD도 포함되어 있습니다.")
-
-LocalizedDisplayNames=(Language="ja",Text="ビークル")
-LocalizedDescriptions=(Language="ja",Text="ビークル パックには通常のビークルとより複雑なビークルが含まれています。キーボード、コントローラー、タッチデバイスでのバーチャルジョイスティックでビークルを操作できます。このビークルにはダブルウィッシュボーン式サスペンションが含まれています。またエンジン音もテンプレートにあります。また現在のギアと速度を示すHUDも用意されています。")
-
-LocalizedDisplayNames=(Language="zh-Hans",Text="载具")
-LocalizedDescriptions=(Language="zh-Hans",Text="载具包包含一辆普通载具和高级载具。你可以通过键盘、手柄或触摸式设备上的虚拟摇杆来控制载具移动。高级载具具有双横臂式悬架。此外，该模板还包含引擎音效，以及一个能够显示当前挡位和车速的HUD。")
-
 ; This project uses a thumbnail as the project icon
 bThumbnailAsIcon=true
 
@@ -53,5 +44,6 @@ EditDetailLevelPreference="High"
 ;SharedContentPacks=(MountName="Variant_Offroad",DetailLevels=("Standard"))
 ;SharedContentPacks=(MountName="Variant_Timetrial",DetailLevels=("Standard"))
 
-Variants=(Name="Timetrial",LocalizedDisplayNames=((Language="en",Text="Timetrial"),(Language="ko",Text="퍼즐")),LocalizedDescriptions=((Language="en",Text="A timed lap with UI and track markers."),(Language="ko",Text="퍼즐 템플릿은 단순한 블록 모눈을 특징으로 합니다.")),SharedContentPacks=((DetailLevels=(Standard),MountName="Variant_Timetrial")) )
-Variants=(Name="Offroad",LocalizedDisplayNames=((Language="en",Text="Offroad"),(Language="ko",Text="퍼즐")),LocalizedDescriptions=((Language="en",Text="An open world map with offroad vehicle."),(Language="ko",Text="퍼즐 템플릿은 단순한 블록 모눈을 특징으로 합니다.")),SharedContentPacks=((DetailLevels=(Standard),MountName="Variant_Offroad")) )
\ No newline at end of file
+Variants=(Name="Timetrial",LocalizedDisplayNames=((Language="en",Text="Timetrial")),LocalizedDescriptions=((Language="en",Text="A timed lap with UI and track markers.")),SharedContentPacks=((DetailLevels=(Standard),MountName="Variant_Timetrial")))
+
+Variants=(Name="Offroad",LocalizedDisplayNames=((Language="en",Text="Offroad")),LocalizedDescriptions=((Language="en",Text="An open world map with offroad vehicle.")),SharedContentPacks=((DetailLevels=(Standard),MountName="Variant_Offroad")))
\ No newline at end of file
