TP_FirstPersonBP
- Commented out variant shared content packs used for maintenance that were incorrectly submitted as uncommented.
#rnx

diff --git a/Templates/TP_FirstPersonBP/Config/TemplateDefs.ini b/Templates/TP_FirstPersonBP/Config/TemplateDefs.ini
--- a/Templates/TP_FirstPersonBP/Config/TemplateDefs.ini
+++ b/Templates/TP_FirstPersonBP/Config/TemplateDefs.ini
@@ -39,9 +39,9 @@ EditDetailLevelPreference="High"
 
 ; Variant packs
 ; NOTE: make sure these are commented out for release!
-SharedContentPacks=(MountName="Variant_Shooter",DetailLevels=("Standard"))
-SharedContentPacks=(MountName="Variant_Horror",DetailLevels=("Standard"))
-SharedContentPacks=(MountName="Weapons",DetailLevels=("Standard"))
+;SharedContentPacks=(MountName="Variant_Shooter",DetailLevels=("Standard"))
+;SharedContentPacks=(MountName="Variant_Horror",DetailLevels=("Standard"))
+;SharedContentPacks=(MountName="Weapons",DetailLevels=("Standard"))
 
 ; Variant definitions
 Variants=(Name="ArenaShooter",LocalizedDisplayNames=((Language="en",Text="Arena Shooter")),LocalizedDescriptions=((Language="en",Text="An arena shooter with weapon pickups and an enemy AI opponent.")),SharedContentPacks=((DetailLevels=(Standard),MountName="Weapons"),(DetailLevels=(Standard),MountName="Variant_Shooter")) )
