Converted Lvl_Horror map to External actors to test one file per actor in resource pack.

Added All option to first person, third person, top down and vehicle that adds all variant options.

diff --git a/Templates/TP_FirstPerson/Config/TemplateDefs.ini b/Templates/TP_FirstPerson/Config/TemplateDefs.ini
--- a/Templates/TP_FirstPerson/Config/TemplateDefs.ini
+++ b/Templates/TP_FirstPerson/Config/TemplateDefs.ini
@@ -3,7 +3,7 @@
 Categories=Games
 
 LocalizedDisplayNames=(Language="en",Text="First Person")
-LocalizedDescriptions=(Language="en",Text="The first person template features a player character (represented by a pair of arms) which is viewed from first person perspective. The character can be moved around the level using keyboard, controller or virtual joystick on a touch device. Additionally the player can look around using mouse, controller or virtual joystick on a touch device. The character can also pick up a gun that, using mouse, controller or virtual joystick on a touch device will fire a simple sphere projectile that will affect some physics objects in the level, whilst bounce off the arena walls.")
+LocalizedDescriptions=(Language="en",Text="The first person template features a player character which is viewed from first person perspective. The character can be moved around the level using keyboard, controller or virtual joystick on a touch device. Additionally the player can look around using mouse, controller or virtual joystick on a touch device. From the variants dropdown you can add - Arena Shooter that adds weapon pickups and a simple AI opponent that will try to shoot you.  - Horror Variant which implements a sprint mechanic and shows best practice for lighting  a dark environment.")
 
 LocalizedDisplayNames=(Language="ko",Text="일인칭")
 LocalizedDescriptions=(Language="ko",Text="일인칭 템플릿은 일인칭 시점에서 보이는 (양팔로 표현되는) 플레이어 캐릭터를 특징으로 합니다. 캐릭터는 레벨 내에서 키보드, 컨트롤러나 터치 디바이스의 가상 조이스틱으로 이동할 수 있습니다. 추가적으로 플레이어는 마우스, 컨트롤러나 터치 디바이스의 가상 조이스틱을 사용하여 주변을 둘러볼 수도 있습니다. 또한, 캐릭터는 총을 장착할 수 있으며 마우스, 컨트롤러나 터치 디바이스의 가상 조이스틱을 사용해 간단한 구형 발사체를 발사하여, 아레나 벽에 튕기면서 레벨에 있는 일부 피직스 오브젝트에 영향을 줍니다.")
diff --git a/Templates/TP_FirstPersonBP/Config/TemplateDefs.ini b/Templates/TP_FirstPersonBP/Config/TemplateDefs.ini
--- a/Templates/TP_FirstPersonBP/Config/TemplateDefs.ini
+++ b/Templates/TP_FirstPersonBP/Config/TemplateDefs.ini
@@ -3,7 +3,7 @@
 Categories=Games
 
 LocalizedDisplayNames=(Language="en",Text="First Person")
-LocalizedDescriptions=(Language="en",Text="The first person template features a player character (represented by a pair of arms) which is viewed from first person perspective. The character can be moved around the level using keyboard, controller or virtual joystick on a touch device. Additionally the player can look around using mouse, controller or virtual joystick on a touch device. The character is also equipped with a gun that, using mouse,  controller or virtual joystick on a touch device will fire a simple sphere projectile that will affect some physics objects in the level, whilst bounce off the arena walls.")
+LocalizedDescriptions=(Language="en",Text="The first person template features a player character which is viewed from first person perspective. The character can be moved around the level using keyboard, controller or virtual joystick on a touch device. Additionally the player can look around using mouse, controller or virtual joystick on a touch device. From the variants dropdown you can add - Arena Shooter that adds weapon pickups and a simple AI opponent that will try to shoot you.  - Horror Variant which implements a sprint mechanic and shows best practice for lighting  a dark environment.")
 
 ; This project uses a thumbnail as the project icon
 bThumbnailAsIcon=true
@@ -47,3 +47,5 @@ EditDetailLevelPreference="High"
 Variants=(Name="ArenaShooter",LocalizedDisplayNames=((Language="en",Text="Arena Shooter")),LocalizedDescriptions=((Language="en",Text="An arena shooter with weapon pickups and an enemy AI opponent.")),SharedContentPacks=((DetailLevels=(Standard),MountName="Weapons"),(DetailLevels=(Standard),MountName="Variant_Shooter")) )
 
 Variants=(Name="SurvivalHorror",LocalizedDisplayNames=((Language="en",Text="Survival Horror")),LocalizedDescriptions=((Language="en",Text="A simple first person character with sprint mechanic and torch set in a dark and suspensful map.")),SharedContentPacks=((DetailLevels=(Standard),MountName="Variant_Horror")) )
+
+Variants=(Name="All",LocalizedDisplayNames=((Language="en",Text="All")),LocalizedDescriptions=((Language="en",Text="Add all the variant content for this template.")),SharedContentPacks=((DetailLevels=(Standard),MountName="Weapons"),(DetailLevels=(Standard),MountName="Variant_Horror"),(DetailLevels=(Standard),MountName="Variant_Shooter")) )
\ No newline at end of file
diff --git a/Templates/TP_ThirdPersonBP/Config/TemplateDefs.ini b/Templates/TP_ThirdPersonBP/Config/TemplateDefs.ini
--- a/Templates/TP_ThirdPersonBP/Config/TemplateDefs.ini
+++ b/Templates/TP_ThirdPersonBP/Config/TemplateDefs.ini
@@ -51,4 +51,6 @@ Variants=(Name="Combat",LocalizedDisplayNames=((Language="en",Text="Combat")),Lo
 
 Variants=(Name="Platforming",LocalizedDisplayNames=((Language="en",Text="Platforming")),LocalizedDescriptions=((Language="en",Text="A 3D platforming game with tweaked movement properties and advanced jump capabilities.")),SharedContentPacks=((DetailLevels=(Standard),MountName="Variant_Platforming")))
 
-Variants=(Name="SideScroller",LocalizedDisplayNames=((Language="en",Text="Side Scroller")),LocalizedDescriptions=((Language="en",Text="A 2.5D platforming game with plane-constrained movement, advanced jump capabilities and a side-scrolling camera.")),SharedContentPacks=((DetailLevels=(Standard),MountName="Variant_SideScroller")))
\ No newline at end of file
+Variants=(Name="SideScroller",LocalizedDisplayNames=((Language="en",Text="Side Scroller")),LocalizedDescriptions=((Language="en",Text="A 2.5D platforming game with plane-constrained movement, advanced jump capabilities and a side-scrolling camera.")),SharedContentPacks=((DetailLevels=(Standard),MountName="Variant_SideScroller")))
+
+Variants=(Name="All",LocalizedDisplayNames=((Language="en",Text="All")),LocalizedDescriptions=((Language="en",Text="Add all the variant content for this template.")),SharedContentPacks=((DetailLevels=(Standard),MountName="Variant_Combat"),(DetailLevels=(Standard),MountName="Variant_Platforming"),(DetailLevels=(Standard),MountName="Variant_SideScroller")))
\ No newline at end of file
diff --git a/Templates/TP_TopDownBP/Config/TemplateDefs.ini b/Templates/TP_TopDownBP/Config/TemplateDefs.ini
--- a/Templates/TP_TopDownBP/Config/TemplateDefs.ini
+++ b/Templates/TP_TopDownBP/Config/TemplateDefs.ini
@@ -50,4 +50,6 @@ EditDetailLevelPreference="High"
 ; Variant definitions
 Variants=(Name="Strategy",LocalizedDisplayNames=((Language="en",Text="Strategy")),LocalizedDescriptions=((Language="en",Text="A top down strategy game with point and click interactions and selectable units.")),SharedContentPacks=((DetailLevels=(Standard),MountName="Variant_Strategy")))
 
-Variants=(Name="TwinStick",LocalizedDisplayNames=((Language="en",Text="Twin Stick")),LocalizedDescriptions=((Language="en",Text="A twin stick shooter game with infinitely spawning enemies.")),SharedContentPacks=((DetailLevels=(Standard),MountName="Variant_TwinStick"),(MountName="Input",DetailLevels=("High")))
\ No newline at end of file
+Variants=(Name="TwinStick",LocalizedDisplayNames=((Language="en",Text="Twin Stick")),LocalizedDescriptions=((Language="en",Text="A twin stick shooter game with infinitely spawning enemies.")),SharedContentPacks=((DetailLevels=(Standard),MountName="Variant_TwinStick"),(MountName="Input",DetailLevels=("High")))
+
+Variants=(Name="All",LocalizedDisplayNames=((Language="en",Text="All")),LocalizedDescriptions=((Language="en",Text="Add all the variant content for this template.")),SharedContentPacks=((DetailLevels=(Standard),MountName="Variant_TwinStick"),(DetailLevels=(Standard),MountName="Variant_Strategy"),(MountName="Input",DetailLevels=("High")))
\ No newline at end of file
diff --git a/Templates/TP_VehicleAdvBP/Config/TemplateDefs.ini b/Templates/TP_VehicleAdvBP/Config/TemplateDefs.ini
--- a/Templates/TP_VehicleAdvBP/Config/TemplateDefs.ini
+++ b/Templates/TP_VehicleAdvBP/Config/TemplateDefs.ini
@@ -47,4 +47,6 @@ EditDetailLevelPreference="High"
 
 Variants=(Name="Timetrial",LocalizedDisplayNames=((Language="en",Text="Timetrial")),LocalizedDescriptions=((Language="en",Text="A timed lap with UI and track markers.")),SharedContentPacks=((DetailLevels=(Standard),MountName="Variant_Timetrial")))
 
-Variants=(Name="Offroad",LocalizedDisplayNames=((Language="en",Text="Offroad")),LocalizedDescriptions=((Language="en",Text="An open world map with offroad vehicle.")),SharedContentPacks=((DetailLevels=(Standard),MountName="Variant_Offroad")))
\ No newline at end of file
+Variants=(Name="Offroad",LocalizedDisplayNames=((Language="en",Text="Offroad")),LocalizedDescriptions=((Language="en",Text="An open world map with offroad vehicle.")),SharedContentPacks=((DetailLevels=(Standard),MountName="Variant_Offroad")))
+
+Variants=(Name="All",LocalizedDisplayNames=((Language="en",Text="All")),LocalizedDescriptions=((Language="en",Text="Add all the variant content for this template.")),SharedContentPacks=((DetailLevels=(Standard),MountName="Variant_Offroad"),(DetailLevels=(Standard),MountName="Variant_Timetrial")))
\ No newline at end of file
diff --git a/Templates/TemplateResources/Standard/Variant_Horror/FeaturePack/manifest.json b/Templates/TemplateResources/Standard/Variant_Horror/FeaturePack/manifest.json
--- a/Templates/TemplateResources/Standard/Variant_Horror/FeaturePack/manifest.json
+++ b/Templates/TemplateResources/Standard/Variant_Horror/FeaturePack/manifest.json
@@ -34,7 +34,9 @@
 	"AdditionalFiles": {
 		"DestinationFilesFolder": "Variant_Horror",
 		"AdditionalFilesList": [
-			"Templates/TemplateResources/Standard/Variant_Horror/Content/*.*"
+			"Templates/TemplateResources/Standard/Variant_Horror/Content/*.*",
+			"Templates/TemplateResources/Standard/Variant_Horror/__ExternalActors__/*.*",
+			"Templates/TemplateResources/Standard/Variant_Horror/__ExternalObjects__/*.*"
 		]
 	}
 }
