Templates
- Updated english descriptions for new templates and content packs.
- Updated feature packs for new templates

#rnx

diff --git a/Templates/TP_FirstPerson/Config/TemplateDefs.ini b/Templates/TP_FirstPerson/Config/TemplateDefs.ini
--- a/Templates/TP_FirstPerson/Config/TemplateDefs.ini
+++ b/Templates/TP_FirstPerson/Config/TemplateDefs.ini
@@ -3,7 +3,7 @@
 Categories=Games
 
 LocalizedDisplayNames=(Language="en",Text="First Person")
-LocalizedDescriptions=(Language="en",Text="The first person template features a player character which is viewed from first person perspective. The character can be moved around the level using keyboard, controller or virtual joystick on a touch device. Additionally the player can look around using mouse, controller or virtual joystick on a touch device. From the variants dropdown you can add - Arena Shooter that adds weapon pickups and a simple AI opponent that will try to shoot you.  - Horror Variant which implements a sprint mechanic and shows best practice for lighting  a dark environment.")
+LocalizedDescriptions=(Language="en",Text="This template features a player character which is viewed from first person perspective. The player can look and move around using mouse, controller or virtual joystick on a touch device. You can select a genre-specific starting point from the Variants drop-down. The C++ version includes code and content for all variants.")
 
 LocalizedDisplayNames=(Language="ko",Text="일인칭")
 LocalizedDescriptions=(Language="ko",Text="일인칭 템플릿은 일인칭 시점에서 보이는 (양팔로 표현되는) 플레이어 캐릭터를 특징으로 합니다. 캐릭터는 레벨 내에서 키보드, 컨트롤러나 터치 디바이스의 가상 조이스틱으로 이동할 수 있습니다. 추가적으로 플레이어는 마우스, 컨트롤러나 터치 디바이스의 가상 조이스틱을 사용하여 주변을 둘러볼 수도 있습니다. 또한, 캐릭터는 총을 장착할 수 있으며 마우스, 컨트롤러나 터치 디바이스의 가상 조이스틱을 사용해 간단한 구형 발사체를 발사하여, 아레나 벽에 튕기면서 레벨에 있는 일부 피직스 오브젝트에 영향을 줍니다.")
@@ -17,8 +17,8 @@ LocalizedDescriptions=(Language="zh-Hans",Text="第一人称模板包含一个
 ; This project uses a thumbnail as the project icon
 bThumbnailAsIcon=true
 
-ClassTypes=Actor, Character, GameMode
-AssetTypes=Animation Sequence, Animation Montage, Sound Wave, Skeletal Mesh, Animation Blueprint, Skeleton, Static Mesh.
+ClassTypes=Actor, Character, GameMode, PlayerController
+AssetTypes=Animation Sequence, Skeletal Mesh, Animation Blueprint, Skeleton, Static Mesh
 
 FoldersToIgnore=Binaries
 FoldersToIgnore=Build
diff --git a/Templates/TP_FirstPerson/manifest.json b/Templates/TP_FirstPerson/manifest.json
--- a/Templates/TP_FirstPerson/manifest.json
+++ b/Templates/TP_FirstPerson/manifest.json
@@ -21,7 +21,7 @@
 	"Description": [
 		{
 			"Language": "en",
-			"Text": "The first person template features a player character (represented by a pair of arms) which is viewed from first person perspective. The character can be moved around the level using keyboard, controller or virtual joystick on a touch device. Additionally the player can look around using mouse, controller or virtual joystick on a touch device. The character can also pick up a gun that, using mouse, controller or virtual joystick on a touch device will fire a simple sphere projectile that will affect some physics objects in the level, whilst bounce off the arena walls."
+			"Text": "This pack features a player character which is viewed from first person perspective. The player can look and move around using mouse, controller or virtual joystick on a touch device."
 		},
 		{
 			"Language": "ko",
@@ -39,7 +39,7 @@
 	"AssetTypes": [
 		{
 			"Language": "en",
-			"Text": "Animation Sequence, Animation Montage, Sound Wave, Skeletal Mesh, Animation Blueprint, Skeleton, Static Mesh."
+			"Text": "Animation Sequence, Skeletal Mesh, Animation Blueprint, Skeleton, Static Mesh"
 		}
 	],
 	"SearchTags": [
@@ -48,7 +48,7 @@
 			"Text": "First Person,Shoot,Fire,Aim"
 		}
 	],
-	"ClassTypes": "Actor, Character, GameMode",
+	"ClassTypes": "Actor, Character, GameMode, PlayerController",
 	"Category": "CodeFeature",
 	"FocusAsset": "/Game/FirstPerson/Blueprints/BP_FirstPersonCharacter.BP_FirstPersonCharacter",
 	"Thumbnail": "TP_FirstPerson.png",
@@ -57,9 +57,9 @@
 	],
 	"AdditionalFeaturePacks": [
 		{
-			"MountName": "FPWeapon",
+			"MountName": "Characters",
 			"DetailLevels": [
-				"Standard"
+				"High"
 			]
 		},
 		{
@@ -69,7 +69,13 @@
 			]
 		},
 		{
-			"MountName": "FirstPersonArms",
+			"MountName": "Input",
+			"DetailLevels": [
+				"High"
+			]
+		},
+		{
+			"MountName": "Weapons",
 			"DetailLevels": [
 				"Standard"
 			]
diff --git a/Templates/TP_FirstPersonBP/Config/TemplateDefs.ini b/Templates/TP_FirstPersonBP/Config/TemplateDefs.ini
--- a/Templates/TP_FirstPersonBP/Config/TemplateDefs.ini
+++ b/Templates/TP_FirstPersonBP/Config/TemplateDefs.ini
@@ -3,13 +3,13 @@
 Categories=Games
 
 LocalizedDisplayNames=(Language="en",Text="First Person")
-LocalizedDescriptions=(Language="en",Text="The first person template features a player character which is viewed from first person perspective. The character can be moved around the level using keyboard, controller or virtual joystick on a touch device. Additionally the player can look around using mouse, controller or virtual joystick on a touch device. From the variants dropdown you can add - Arena Shooter that adds weapon pickups and a simple AI opponent that will try to shoot you.  - Horror Variant which implements a sprint mechanic and shows best practice for lighting  a dark environment.")
+LocalizedDescriptions=(Language="en",Text="This template features a player character which is viewed from first person perspective. The player can look and move around using mouse, controller or virtual joystick on a touch device. You can select a genre-specific starting point from the Variants drop-down. The C++ version includes code and content for all variants.")
 
 ; This project uses a thumbnail as the project icon
 bThumbnailAsIcon=true
 
-ClassTypes=Character, GameMode, HUD, Actor
-AssetTypes=Animation Sequence, Animation Montage, Sound Wave, Skeletal Mesh, Animation Blueprint, Skeleton, Static Mesh, Texture.
+ClassTypes=Actor, Character, GameMode, PlayerController
+AssetTypes=Animation Sequence, Skeletal Mesh, Animation Blueprint, Skeleton, Static Mesh
 
 FoldersToIgnore=Binaries
 FoldersToIgnore=Build
@@ -46,8 +46,8 @@ EditDetailLevelPreference="High"
 ;SharedContentPacks=(MountName="Weapons",DetailLevels=("Standard"))
 
 ; Variant definitions
-Variants=(Name="ArenaShooter",LocalizedDisplayNames=((Language="en",Text="Arena Shooter")),LocalizedDescriptions=((Language="en",Text="An arena shooter with weapon pickups and an enemy AI opponent.")),SharedContentPacks=((DetailLevels=(Standard),MountName="Weapons"),(DetailLevels=(Standard),MountName="Variant_Shooter")) )
+Variants=(Name="ArenaShooter",LocalizedDisplayNames=((Language="en",Text="Arena Shooter")),LocalizedDescriptions=((Language="en",Text="An arena shooter with weapon pickups and enemy AI opponents.")),SharedContentPacks=((DetailLevels=(Standard),MountName="Weapons"),(DetailLevels=(Standard),MountName="Variant_Shooter")) )
 
-Variants=(Name="SurvivalHorror",LocalizedDisplayNames=((Language="en",Text="Survival Horror")),LocalizedDescriptions=((Language="en",Text="A simple first person character with sprint mechanic and torch set in a dark and suspensful map.")),SharedContentPacks=((DetailLevels=(Standard),MountName="Variant_Horror")) )
+Variants=(Name="SurvivalHorror",LocalizedDisplayNames=((Language="en",Text="Survival Horror")),LocalizedDescriptions=((Language="en",Text="A simple first person character with sprint mechanic and torch set in a dark and suspenseful map.")),SharedContentPacks=((DetailLevels=(Standard),MountName="Variant_Horror")) )
 
-Variants=(Name="All",LocalizedDisplayNames=((Language="en",Text="All")),LocalizedDescriptions=((Language="en",Text="Add all the variant content for this template.")),SharedContentPacks=((DetailLevels=(Standard),MountName="Weapons"),(DetailLevels=(Standard),MountName="Variant_Horror"),(DetailLevels=(Standard),MountName="Variant_Shooter")) )
\ No newline at end of file
+Variants=(Name="All",LocalizedDisplayNames=((Language="en",Text="All")),LocalizedDescriptions=((Language="en",Text="Includes all First Person variants.")),SharedContentPacks=((DetailLevels=(Standard),MountName="Weapons"),(DetailLevels=(Standard),MountName="Variant_Horror"),(DetailLevels=(Standard),MountName="Variant_Shooter")) )
\ No newline at end of file
diff --git a/Templates/TP_FirstPersonBP/manifest.json b/Templates/TP_FirstPersonBP/manifest.json
--- a/Templates/TP_FirstPersonBP/manifest.json
+++ b/Templates/TP_FirstPersonBP/manifest.json
@@ -21,7 +21,7 @@
 	"Description": [
 		{
 			"Language": "en",
-			"Text": "The first person template features a player character (represented by a pair of arms) which is viewed from first person perspective. The character can be moved around the level using keyboard, controller or virtual joystick on a touch device. Additionally the player can look around using mouse, controller or virtual joystick on a touch device. The character can also pick up a gun that, using mouse, controller or virtual joystick on a touch device will fire a simple sphere projectile that will affect some physics objects in the level, whilst bounce off the arena walls."
+			"Text": "This pack features a player character which is viewed from first person perspective. The player can look and move around using mouse, controller or virtual joystick on a touch device."
 		},
 		{
 			"Language": "ko",
@@ -39,7 +39,7 @@
 	"AssetTypes": [
 		{
 			"Language": "en",
-			"Text": "Animation Sequence, Animation Montage, Sound Wave, Skeletal Mesh, Animation Blueprint, Skeleton, Static Mesh."
+			"Text": "Animation Sequence, Skeletal Mesh, Animation Blueprint, Skeleton, Static Mesh"
 		}
 	],
 	"SearchTags": [
@@ -48,7 +48,7 @@
 			"Text": "First Person,Shoot,Fire,Aim"
 		}
 	],
-	"ClassTypes": "Actor, Character, GameMode",
+	"ClassTypes": "Actor, Character, GameMode, PlayerController",
 	"Category": "BlueprintFeature",
 	"FocusAsset": "/Game/FirstPerson/Blueprints/BP_FirstPersonCharacter.BP_FirstPersonCharacter",
 	"Thumbnail": "TP_FirstPersonBP.png",
@@ -57,13 +57,13 @@
 	],
 	"AdditionalFeaturePacks": [
 		{
-			"MountName": "LevelPrototyping",
+			"MountName": "Characters",
 			"DetailLevels": [
 				"High"
 			]
 		},
 		{
-			"MountName": "Characters",
+			"MountName": "LevelPrototyping",
 			"DetailLevels": [
 				"High"
 			]
@@ -73,6 +73,12 @@
 			"DetailLevels": [
 				"High"
 			]
+		},
+		{
+			"MountName": "Weapons",
+			"DetailLevels": [
+				"Standard"
+			]
 		}
 	]
 }
diff --git a/Templates/TP_ThirdPerson/Config/TemplateDefs.ini b/Templates/TP_ThirdPerson/Config/TemplateDefs.ini
--- a/Templates/TP_ThirdPerson/Config/TemplateDefs.ini
+++ b/Templates/TP_ThirdPerson/Config/TemplateDefs.ini
@@ -4,7 +4,7 @@ Categories=Games
 
 
 LocalizedDisplayNames=(Language="en",Text="Third Person")
-LocalizedDescriptions=(Language="en",Text="The third person pack features a playable character where the camera is positioned behind and slightly above the character. As the character moves using mouse, keyboard, controller or virtual joystick on a touch device, the camera follows the character, observing from an over-the-shoulder point of view. This perspective emphasizes the main character and is popular in action and adventure games. The template map contains several objects that the player can jump upon, and push around.")
+LocalizedDescriptions=(Language="en",Text="This template features a playable character where the camera observes an over-the-shoulder point of view. The player can rotate the camera and move around using keyboard, mouse, controller or virtual joystick on a touch device. You can select a genre-specific starting point from the Variants drop-down. The C++ version includes code and content for all variants.")
 
 LocalizedDisplayNames=(Language="ko",Text="삼인칭")
 LocalizedDescriptions=(Language="ko",Text="삼인칭 템플릿에는 플레이 가능한 캐릭터가 포함되어 있으며 카메라가 캐릭터 머리 위 약간 뒤쪽에 위치합니다. 캐릭터가 마우스, 키보드, 터치 디바이스의 가상 조이스틱으로 움직일 때 카메라는 어깨너머의 관점에서 캐릭터를 관찰하며 따라갑니다. 이 시점은 주인공 캐릭터를 강조해 액션이나 어드벤처 게임에 주로 활용됩니다. 이 템플릿 맵에는 플레이어가 점프하여 착지하고 밀 수 있는 여러 가지 오브젝트가 포함되어 있습니다.")
@@ -18,8 +18,8 @@ LocalizedDescriptions=(Language="zh-Hans",Text="第三人称游戏内容包包
 ; This project uses a thumbnail as the project icon
 bThumbnailAsIcon=true
 
-ClassTypes=Character, SpringArmComponent, GameMode
-AssetTypes=Animation Sequence, Animation Blueprint, Skeleton, Skeletal Mesh, Control Rig, IK Rig, Retargeter
+ClassTypes=Character, SpringArmComponent, PlayerController, GameMode
+AssetTypes=Animation Sequence, Animation Blueprint, Skeleton, Skeletal Mesh, Control Rig
 
 FoldersToIgnore=Binaries
 FoldersToIgnore=Build
diff --git a/Templates/TP_ThirdPerson/manifest.json b/Templates/TP_ThirdPerson/manifest.json
--- a/Templates/TP_ThirdPerson/manifest.json
+++ b/Templates/TP_ThirdPerson/manifest.json
@@ -21,7 +21,7 @@
 	"Description": [
 		{
 			"Language": "en",
-			"Text": "The third person pack features a playable character where the camera is positioned behind and slightly above the character. As the character moves using mouse, keyboard, controller or virtual joystick on a touch device, the camera follows the character, observing from an over-the-shoulder point of view. This perspective emphasizes the main character and is popular in action and adventure games. The template map contains several objects that the player can jump upon, and push around."
+			"Text": "This pack features a playable character where the camera observes an over-the-shoulder point of view. The player can rotate the camera and move around using keyboard, mouse, controller or virtual joystick on a touch device."
 		},
 		{
 			"Language": "ko",
@@ -39,7 +39,7 @@
 	"AssetTypes": [
 		{
 			"Language": "en",
-			"Text": "Animation Sequence, Animation Blueprint, Skeleton, Skeletal Mesh, Control Rig, IK Rig, Retargeter"
+			"Text": "Animation Sequence, Animation Blueprint, Skeleton, Skeletal Mesh, Control Rig"
 		}
 	],
 	"SearchTags": [
@@ -48,7 +48,7 @@
 			"Text": "Action, Adventure, Third"
 		}
 	],
-	"ClassTypes": "Character, SpringArmComponent, GameMode",
+	"ClassTypes": "Character, SpringArmComponent, PlayerController, GameMode",
 	"Category": "CodeFeature",
 	"FocusAsset": "/Game/ThirdPerson/Blueprints/BP_ThirdPersonCharacter.BP_ThirdPersonCharacter",
 	"Thumbnail": "TP_ThirdPerson.png",
@@ -67,6 +67,12 @@
 			"DetailLevels": [
 				"High"
 			]
+		},
+		{
+			"MountName": "Input",
+			"DetailLevels": [
+				"High"
+			]
 		}
 	]
 }
diff --git a/Templates/TP_ThirdPersonBP/Config/TemplateDefs.ini b/Templates/TP_ThirdPersonBP/Config/TemplateDefs.ini
--- a/Templates/TP_ThirdPersonBP/Config/TemplateDefs.ini
+++ b/Templates/TP_ThirdPersonBP/Config/TemplateDefs.ini
@@ -3,13 +3,13 @@
 Categories=Games
 
 LocalizedDisplayNames=(Language="en",Text="Third Person")
-LocalizedDescriptions=(Language="en",Text="The third person pack features a playable character where the camera is positioned behind and slightly above the character. As the character moves using mouse, keyboard, controller or virtual joystick on a touch device, the camera follows the character, observing from an over-the-shoulder point of view. Using variants, this template can be expanded to work as a 2.5D sidescroller, a traditional platformer or a hand to hand combat game.")
+LocalizedDescriptions=(Language="en",Text="This template features a playable character where the camera observes an over-the-shoulder point of view. The player can rotate the camera and move around using keyboard, mouse, controller or virtual joystick on a touch device. You can select a genre-specific starting point from the Variants drop-down. The C++ version includes code and content for all variants.")
 
 ; This project uses a thumbnail as the project icon
 bThumbnailAsIcon=true
 
-ClassTypes=Character, SpringArmComponent, GameMode
-AssetTypes=Animation Sequence, Animation Blueprint, Skeleton, Skeletal Mesh, Control Rig, IK Rig, Retargeter
+ClassTypes=Character, SpringArmComponent, PlayerController, GameMode
+AssetTypes=Animation Sequence, Animation Blueprint, Skeleton, Skeletal Mesh, Control Rig
 
 FoldersToIgnore=Binaries
 FoldersToIgnore=Build
@@ -53,4 +53,4 @@ Variants=(Name="Platforming",LocalizedDisplayNames=((Language="en",Text="Platfor
 
 Variants=(Name="SideScroller",LocalizedDisplayNames=((Language="en",Text="Side Scroller")),LocalizedDescriptions=((Language="en",Text="A 2.5D platforming game with plane-constrained movement, advanced jump capabilities and a side-scrolling camera.")),SharedContentPacks=((DetailLevels=(Standard),MountName="Variant_SideScroller")))
 
-Variants=(Name="All",LocalizedDisplayNames=((Language="en",Text="All")),LocalizedDescriptions=((Language="en",Text="Add all the variant content for this template.")),SharedContentPacks=((DetailLevels=(Standard),MountName="Variant_Combat"),(DetailLevels=(Standard),MountName="Variant_Platforming"),(DetailLevels=(Standard),MountName="Variant_SideScroller")))
\ No newline at end of file
+Variants=(Name="All",LocalizedDisplayNames=((Language="en",Text="All")),LocalizedDescriptions=((Language="en",Text="Includes all Third Person variants.")),SharedContentPacks=((DetailLevels=(Standard),MountName="Variant_Combat"),(DetailLevels=(Standard),MountName="Variant_Platforming"),(DetailLevels=(Standard),MountName="Variant_SideScroller")))
\ No newline at end of file
diff --git a/Templates/TP_ThirdPersonBP/manifest.json b/Templates/TP_ThirdPersonBP/manifest.json
--- a/Templates/TP_ThirdPersonBP/manifest.json
+++ b/Templates/TP_ThirdPersonBP/manifest.json
@@ -21,7 +21,7 @@
 	"Description": [
 		{
 			"Language": "en",
-			"Text": "The third person pack features a playable character where the camera is positioned behind and slightly above the character. As the character moves using mouse, keyboard, controller or virtual joystick on a touch device, the camera follows the character, observing from an over-the-shoulder point of view. This perspective emphasizes the main character and is popular in action and adventure games. The template map contains several objects that the player can jump upon, and push around."
+			"Text": "This pack features a playable character where the camera observes an over-the-shoulder point of view. The player can rotate the camera and move around using keyboard, mouse, controller or virtual joystick on a touch device."
 		},
 		{
 			"Language": "ko",
@@ -39,7 +39,7 @@
 	"AssetTypes": [
 		{
 			"Language": "en",
-			"Text": "Animation Sequence, Animation Blueprint, Skeleton, Skeletal Mesh, Control Rig, IK Rig, Retargeter"
+			"Text": "Animation Sequence, Animation Blueprint, Skeleton, Skeletal Mesh, Control Rig"
 		}
 	],
 	"SearchTags": [
@@ -48,7 +48,7 @@
 			"Text": "Action, Adventure, Third"
 		}
 	],
-	"ClassTypes": "Character, SpringArmComponent, GameMode",
+	"ClassTypes": "Character, SpringArmComponent, PlayerController, GameMode",
 	"Category": "BlueprintFeature",
 	"FocusAsset": "/Game/ThirdPerson/Blueprints/BP_ThirdPersonCharacter.BP_ThirdPersonCharacter",
 	"Thumbnail": "TP_ThirdPersonBP.png",
diff --git a/Templates/TP_TopDown/Config/TemplateDefs.ini b/Templates/TP_TopDown/Config/TemplateDefs.ini
--- a/Templates/TP_TopDown/Config/TemplateDefs.ini
+++ b/Templates/TP_TopDown/Config/TemplateDefs.ini
@@ -4,7 +4,7 @@ Categories=Games
 
 
 LocalizedDisplayNames=(Language="en",Text="Top Down")
-LocalizedDescriptions=(Language="en",Text="The top down template features a playable character where the camera is positioned behind and above the avatar at greater distance than in third person perspective. The player is controlled by using the mouse or touchscreen to click on the required destination and uses the navigation system to facilitate the character movement. This perspective is often used in action role playing games.")
+LocalizedDescriptions=(Language="en",Text="This template template features a playable character and a top-down camera perspective at a greater distance. The player can control the character by clicking or touching on a destination, or following the cursor if input is held. You can select a genre-specific starting point from the Variants drop-down. The C++ version includes code and content for all variants.")
 
 LocalizedDisplayNames=(Language="ko",Text="내려보기")
 LocalizedDescriptions=(Language="ko",Text="내려보기 템플릿은 플레이어 캐릭터 위 삼인칭 시점보다 멀리 카메라가 위치합니다. 플레이어 제어는 마우스나 터치스크린으로 원하는 목적지를 클릭하는 식으로 이루어지며, 원활한 캐릭터 이동을 위해 내비게이션 시스템을 사용합니다. 이 시점은 액션 롤플레잉 게임에 자주 쓰입니다.")
@@ -19,7 +19,7 @@ LocalizedDescriptions=(Language="zh-Hans",Text="俯视角游戏模板拥有一
 bThumbnailAsIcon=true
 
 ClassTypes=GameMode, Character, CameraComponent, SpringArmComponent, PlayerController
-AssetTypes=Animation Sequence, Skeletal Mesh, Animation Blueprint,Skeleton, Blend Space 1D
+AssetTypes=Animation Sequence, Skeletal Mesh, Animation Blueprint, Skeleton
 
 FoldersToIgnore=Binaries
 FoldersToIgnore=Build
diff --git a/Templates/TP_TopDown/manifest.json b/Templates/TP_TopDown/manifest.json
--- a/Templates/TP_TopDown/manifest.json
+++ b/Templates/TP_TopDown/manifest.json
@@ -21,7 +21,7 @@
 	"Description": [
 		{
 			"Language": "en",
-			"Text": "The top down pack features a playable character where the camera is positioned behind and above the avatar at greater distance than in third person perspective. The player is controlled by using the mouse or touchscreen to click on the required destination and uses the navigation system to facilitate the character movement. This perspective is often used in action role playing games."
+			"Text": "This pack features a playable character and a top-down camera perspective at a greater distance. The player can control the character by clicking or touching on a destination, or following the cursor if input is held."
 		},
 		{
 			"Language": "ko",
@@ -39,7 +39,7 @@
 	"AssetTypes": [
 		{
 			"Language": "en",
-			"Text": "Animation Sequence, Skeletal Mesh, Animation Blueprint,Skeleton, Blend Space 1D"
+			"Text": "Animation Sequence, Skeletal Mesh, Animation Blueprint, Skeleton"
 		}
 	],
 	"SearchTags": [
diff --git a/Templates/TP_TopDownBP/Config/TemplateDefs.ini b/Templates/TP_TopDownBP/Config/TemplateDefs.ini
--- a/Templates/TP_TopDownBP/Config/TemplateDefs.ini
+++ b/Templates/TP_TopDownBP/Config/TemplateDefs.ini
@@ -4,13 +4,13 @@ Categories=Games
 
 
 LocalizedDisplayNames=(Language="en",Text="Top Down")
-LocalizedDescriptions=(Language="en",Text="The top down template features a playable character where the camera is positioned behind and above the avatar at greater distance than in third person perspective. The player is controlled by using the mouse, gamepad, or touchscreen and moves the character by clicking on a destination, or following the cursor if input is held. This perspective is often used in action role playing games.")
+LocalizedDescriptions=(Language="en",Text="This template template features a playable character and a top-down camera perspective at a greater distance. The player can control the character by clicking or touching on a destination, or following the cursor if input is held. You can select a genre-specific starting point from the Variants drop-down. The C++ version includes code and content for all variants.")
 
 ; This project uses a thumbnail as the project icon
 bThumbnailAsIcon=true
 
-ClassTypes=Character, PlayerController, CameraComponent, SpringArmComponent, GameMode
-AssetTypes=Animation Sequence, Skeletal Mesh, Animation Blueprint,Skeleton
+ClassTypes=GameMode, Character, CameraComponent, SpringArmComponent, PlayerController
+AssetTypes=Animation Sequence, Skeletal Mesh, Animation Blueprint, Skeleton
 
 FoldersToIgnore=Binaries
 FoldersToIgnore=Build
@@ -52,4 +52,4 @@ Variants=(Name="Strategy",LocalizedDisplayNames=((Language="en",Text="Strategy")
 
 Variants=(Name="TwinStick",LocalizedDisplayNames=((Language="en",Text="Twin Stick")),LocalizedDescriptions=((Language="en",Text="A twin stick shooter game with infinitely spawning enemies.")),SharedContentPacks=((DetailLevels=(Standard),MountName="Variant_TwinStick"),(MountName="Input",DetailLevels=("High")))
 
-Variants=(Name="All",LocalizedDisplayNames=((Language="en",Text="All")),LocalizedDescriptions=((Language="en",Text="Add all the variant content for this template.")),SharedContentPacks=((DetailLevels=(Standard),MountName="Variant_TwinStick"),(DetailLevels=(Standard),MountName="Variant_Strategy"),(MountName="Input",DetailLevels=("High")))
\ No newline at end of file
+Variants=(Name="All",LocalizedDisplayNames=((Language="en",Text="All")),LocalizedDescriptions=((Language="en",Text="Includes all Top Down variants.")),SharedContentPacks=((DetailLevels=(Standard),MountName="Variant_TwinStick"),(DetailLevels=(Standard),MountName="Variant_Strategy"),(MountName="Input",DetailLevels=("High")))
\ No newline at end of file
diff --git a/Templates/TP_TopDownBP/manifest.json b/Templates/TP_TopDownBP/manifest.json
--- a/Templates/TP_TopDownBP/manifest.json
+++ b/Templates/TP_TopDownBP/manifest.json
@@ -21,7 +21,7 @@
 	"Description": [
 		{
 			"Language": "en",
-			"Text": "The top down template features a playable character where the camera is positioned behind and above the avatar at greater distance than in third person perspective. The player is controlled by using the mouse, gamepad, or touchscreen and moves the character by clicking on a destination, or following the cursor if input is held. This perspective is often used in action role playing games."
+			"Text": "This pack features a playable character and a top-down camera perspective at a greater distance. The player can control the character by clicking or touching on a destination, or following the cursor if input is held."
 		},
 		{
 			"Language": "ko",
@@ -48,7 +48,7 @@
 			"Text": "RPG,Role Playing,Top Down"
 		}
 	],
-	"ClassTypes": "Character, PlayerController, CameraComponent, SpringArmComponent, GameMode",
+	"ClassTypes": "GameMode, Character, CameraComponent, SpringArmComponent, PlayerController",
 	"Category": "BlueprintFeature",
 	"FocusAsset": "/Game/TopDown/Blueprints/BP_TopDownCharacter.BP_TopDownCharacter",
 	"Thumbnail": "TP_TopDownBP.png",
diff --git a/Templates/TP_VehicleAdv/Config/TemplateDefs.ini b/Templates/TP_VehicleAdv/Config/TemplateDefs.ini
--- a/Templates/TP_VehicleAdv/Config/TemplateDefs.ini
+++ b/Templates/TP_VehicleAdv/Config/TemplateDefs.ini
@@ -3,7 +3,7 @@
 Categories=Games
 
 LocalizedDisplayNames=(Language="en",Text="Vehicle")
-LocalizedDescriptions=(Language="en",Text="The vehicle pack features both a regular vehicle and a more complex vehicle. Vehicle movement can be controlled with either keyboard, controller or virtual joystick on a touch device. The advanced vehicle itself features working double wishbone suspension. There is also an engine sound in this template. There is also a HUD that shows the current gear and speed.")
+LocalizedDescriptions=(Language="en",Text="This template features a sports car and an offroad vehicle with a double wishbone suspension, as well as a HUD displaying the current gear and speed. Vehicle movement can be controlled with keyboard, controller, steering wheel, or virtual joystick on a touch device. You can select a genre-specific starting point from the Variants drop-down. The C++ version includes code and content for all variants.")
 
 LocalizedDisplayNames=(Language="ko",Text="비히클")
 LocalizedDescriptions=(Language="ko",Text="비히클 팩은 일반 비히클과 조금 더 복잡한 비히클을 모두 제공합니다. 비히클의 이동은 키보드, 컨트롤러 또는 터치 디바이스의 버추얼 조이스틱으로 제어 가능합니다. 고급 비히클은 더블 위시본 서스펜션을 갖추고 있습니다. 이 템플릿에는 엔진 사운드와, 현재 기어와 속도를 보여주는 HUD도 포함되어 있습니다.")
@@ -17,8 +17,8 @@ LocalizedDescriptions=(Language="zh-Hans",Text="载具包包含一辆普通载
 ; This project uses a thumbnail as the project icon
 bThumbnailAsIcon=true
 
-ClassTypes=GameMode, WheeledVehicle, HUD, SpringArmComponent, Font, CameraComponent, TextRenderComponent, InputComponent, PhysicalMaterial, AudioComponent
-AssetTypes=Static Mesh, Landscape Layer, Material, Instance, Physical Material, Sound Attenuation, Sound cue, Sound Wave, Texture, Skeletal Mesh, Skeleton, Physics Asset, Animation blueprint, Data Asset, Vehicle Wheel
+ClassTypes=GameMode, WheeledVehicle, HUD, SpringArmComponent, CameraComponent, PhysicalMaterial
+AssetTypes=Static Mesh, Landscape, Material, Instance, Physical Material, Texture, Skeletal Mesh, Skeleton, Physics Asset, Animation Blueprint, Data Asset, Vehicle Wheel
 
 FoldersToIgnore=Binaries
 FoldersToIgnore=Build
@@ -52,6 +52,4 @@ ReplacementsInFiles=(Extensions=("cpp","h","ini","cs"),From="%TEMPLATENAME%",To=
 
 SortKey=ZVehicle_Adv
 SharedContentPacks=(MountName="Vehicles",DetailLevels=("Standard"))
-SharedContentPacks=(MountName="Track",DetailLevels=("High"))
-SharedContentPacks=(MountName="LevelPrototyping",DetailLevels=("High"))
 SharedContentPacks=(MountName="Input",DetailLevels=("High"))
\ No newline at end of file
diff --git a/Templates/TP_VehicleAdv/manifest.json b/Templates/TP_VehicleAdv/manifest.json
--- a/Templates/TP_VehicleAdv/manifest.json
+++ b/Templates/TP_VehicleAdv/manifest.json
@@ -23,7 +23,7 @@
 	[
 		{
 			"Language": "en",
-			"Text": "The vehicle pack features both a regular vehicle and a more complex vehicle. Vehicle movement can be controlled with either keyboard, controller or virtual joystick on a touch device. The advanced vehicle itself features working double wishbone suspension. There is also an engine sound in this template. There is also a HUD that shows the current gear and speed."
+			"Text": "This pack features a sports car and an offroad vehicle with a double wishbone suspension, as well as a HUD displaying the current gear and speed. Vehicle movement can be controlled with keyboard, controller, steering wheel, or virtual joystick on a touch device."
 		},
 		{
 			"Language": "ja",
@@ -42,7 +42,7 @@
 	[
 		{
 			"Language": "en",
-			"Text": "Static Mesh, Landscape Layer, Material, Instance, Physical Material, Sound Attenuation, Sound cue, Sound Wave, Texture, Skeletal Mesh, Skeleton, Physics Asset, Animation blueprint, Data Asset, Vehicle Wheel"
+			"Text": "Static Mesh, Landscape, Material, Instance, Physical Material, Texture, Skeletal Mesh, Skeleton, Physics Asset, Animation Blueprint, Data Asset, Vehicle Wheel"
 		}
 	],	
 	"SearchTags":
@@ -52,7 +52,7 @@
 			"Text": "Tyre,Suspension,Wheel,Racing,Vehicle,Car"
 		}
 	],
-	"ClassTypes": "GameMode, WheeledVehiclePawn, HUD, SpringArmComponent, Font, CameraComponent, TextRenderComponent, InputComponent, PhysicalMaterial, AudioComponent",
+	"ClassTypes": "GameMode, WheeledVehicle, HUD, SpringArmComponent, CameraComponent, PhysicalMaterial",
 	"Category": "CodeFeature",
 	"FocusAsset": "/Game/VehicleAdvCPP/AdvancedVehicleOverview.AdvancedVehicleOverview",
 	"Thumbnail": "TP_VehicleAdv.png",
@@ -71,16 +71,8 @@
 			]
 		},
 		{
-			"MountName": "Track",
-			"DetailLevels":
-			[
-				"High"
-			]
-		},
-		{
-			"MountName": "LevelPrototyping",
-			"DetailLevels":
-			[
+			"MountName": "Input",
+			"DetailLevels": [
 				"High"
 			]
 		}
diff --git a/Templates/TP_VehicleAdvBP/Config/TemplateDefs.ini b/Templates/TP_VehicleAdvBP/Config/TemplateDefs.ini
--- a/Templates/TP_VehicleAdvBP/Config/TemplateDefs.ini
+++ b/Templates/TP_VehicleAdvBP/Config/TemplateDefs.ini
@@ -3,13 +3,13 @@
 Categories=Games
 
 LocalizedDisplayNames=(Language="en",Text="Vehicle")
-LocalizedDescriptions=(Language="en",Text="The vehicle pack features both a regular vehicle and a more complex vehicle. Vehicle movement can be controlled with either keyboard, controller or virtual joystick on a touch device. The advanced vehicle itself features working double wishbone suspension. There is also an engine sound in this template. There is also a HUD that shows the current gear and speed.")
+LocalizedDescriptions=(Language="en",Text="This template features a sports car and an offroad vehicle with a double wishbone suspension, as well as a HUD displaying the current gear and speed. Vehicle movement can be controlled with keyboard, controller, steering wheel, or virtual joystick on a touch device. You can select a genre-specific starting point from the Variants drop-down. The C++ version includes code and content for all variants.")
 
 ; This project uses a thumbnail as the project icon
 bThumbnailAsIcon=true
 
-ClassTypes=GameMode, WheeledVehicle, HUD, SpringArmComponent, Font, CameraComponent, TextRenderComponent, InputComponent, PhysicalMaterial, AudioComponent
-AssetTypes=Static Mesh, Landscape Layer, Material, Instance, Physical Material, Sound Attenuation, Sound cue, Sound Wave, Texture, Skeletal Mesh, Skeleton, Physics Asset, Animation blueprint, Data Asset, Vehicle Wheel
+ClassTypes=GameMode, WheeledVehicle, HUD, SpringArmComponent, CameraComponent, PhysicalMaterial
+AssetTypes=Static Mesh, Landscape, Material, Instance, Physical Material, Texture, Skeletal Mesh, Skeleton, Physics Asset, Animation Blueprint, Data Asset, Vehicle Wheel
 
 FoldersToIgnore=Binaries
 FoldersToIgnore=Build
@@ -45,8 +45,8 @@ EditDetailLevelPreference="High"
 ;SharedContentPacks=(MountName="Variant_Offroad",DetailLevels=("Standard"))
 ;SharedContentPacks=(MountName="Variant_Timetrial",DetailLevels=("Standard"))
 
-Variants=(Name="Timetrial",LocalizedDisplayNames=((Language="en",Text="Timetrial")),LocalizedDescriptions=((Language="en",Text="A timed lap with UI and track markers.")),SharedContentPacks=((DetailLevels=(Standard),MountName="Variant_Timetrial")))
+Variants=(Name="Timetrial",LocalizedDisplayNames=((Language="en",Text="Time Trial")),LocalizedDescriptions=((Language="en",Text="A timed lap with UI and track markers.")),SharedContentPacks=((DetailLevels=(Standard),MountName="Variant_Timetrial")))
 
-Variants=(Name="Offroad",LocalizedDisplayNames=((Language="en",Text="Offroad")),LocalizedDescriptions=((Language="en",Text="An open world map with offroad vehicle.")),SharedContentPacks=((DetailLevels=(Standard),MountName="Variant_Offroad")))
+Variants=(Name="Offroad",LocalizedDisplayNames=((Language="en",Text="Offroad")),LocalizedDescriptions=((Language="en",Text="An open world map with an advanced offroad vehicle.")),SharedContentPacks=((DetailLevels=(Standard),MountName="Variant_Offroad")))
 
-Variants=(Name="All",LocalizedDisplayNames=((Language="en",Text="All")),LocalizedDescriptions=((Language="en",Text="Add all the variant content for this template.")),SharedContentPacks=((DetailLevels=(Standard),MountName="Variant_Offroad"),(DetailLevels=(Standard),MountName="Variant_Timetrial")))
\ No newline at end of file
+Variants=(Name="All",LocalizedDisplayNames=((Language="en",Text="All")),LocalizedDescriptions=((Language="en",Text="Includes all Vehicle variants.")),SharedContentPacks=((DetailLevels=(Standard),MountName="Variant_Offroad"),(DetailLevels=(Standard),MountName="Variant_Timetrial")))
\ No newline at end of file
diff --git a/Templates/TP_VehicleAdvBP/manifest.json b/Templates/TP_VehicleAdvBP/manifest.json
--- a/Templates/TP_VehicleAdvBP/manifest.json
+++ b/Templates/TP_VehicleAdvBP/manifest.json
@@ -21,7 +21,7 @@
 	"Description": [
 		{
 			"Language": "en",
-			"Text": "The vehicle pack features both a regular vehicle and a more complex vehicle. Vehicle movement can be controlled with either keyboard, controller or virtual joystick on a touch device. The advanced vehicle itself features working double wishbone suspension. There is also an engine sound in this template. There is also a HUD that shows the current gear and speed."
+			"Text": "This pack features a sports car and an offroad vehicle with a double wishbone suspension, as well as a HUD displaying the current gear and speed. Vehicle movement can be controlled with keyboard, controller, steering wheel, or virtual joystick on a touch device."
 		},
 		{
 			"Language": "ja",
@@ -39,7 +39,7 @@
 	"AssetTypes": [
 		{
 			"Language": "en",
-			"Text": "Static Mesh, Landscape Layer, Material, Instance, Physical Material, Sound Attenuation, Sound cue, Sound Wave, Texture, Skeletal Mesh, Skeleton, Physics Asset, Animation blueprint, Data Asset, Vehicle Wheel"
+			"Text": "Static Mesh, Landscape, Material, Instance, Physical Material, Texture, Skeletal Mesh, Skeleton, Physics Asset, Animation Blueprint, Data Asset, Vehicle Wheel"
 		}
 	],
 	"SearchTags": [
@@ -48,7 +48,7 @@
 			"Text": "Tyre,Suspension,Wheel,Racing,Vehicle,Car"
 		}
 	],
-	"ClassTypes": "GameMode, WheeledVehicle, HUD, SpringArmComponent, Font, CameraComponent, TextRenderComponent, InputComponent, PhysicalMaterial, AudioComponent",
+	"ClassTypes": "GameMode, WheeledVehicle, HUD, SpringArmComponent, CameraComponent, PhysicalMaterial",
 	"Category": "BlueprintFeature",
 	"FocusAsset": "/Game/VehicleTemplate/Blueprints/VehiclePlayerController.VehiclePlayerController",
 	"Thumbnail": "TP_VehicleAdvBP.png",
